UE::WorldMetrics plugin implementation:
- WorldMetricsSubsystem
- WorldMetric
- WorldMetricsExtension Acquire/Release semantics
- WorldMetricTest: unit test collection to validate the UWorldMetricsSubsystem metric and extension framework. This is a separate module excluded from shipping builds.
- CsvActorCountMetric implementing CSV Profiler stat
- CsvMetricsSubsystem automatically adds/removes CsvActorCountMetric when a CSV profiling starts/ends.

#tests World metric unit tests, in-game CSvActorCountMetric tested through CSVprofile start/end. Faked susbsystem deinitialization failure test to ensure that enabled subsystems in an invalid state correctly unregister their ticker during PreDestroy.
#design https://docs.google.com/presentation/d/1btGRHfrjy2R4lL-2wI7CEdyIAQfQIqex_KQ83XsClb8/edit?usp=sharing
#prebuild 65cb39ea4e11958de195833d

diff --git a/Engine/Plugins/WorldMetrics/Source/CsvMetrics/CsvMetrics.Build.cs b/Engine/Plugins/WorldMetrics/Source/CsvMetrics/CsvMetrics.Build.cs
new file mode 100644
--- /dev/null
+++ b/Engine/Plugins/WorldMetrics/Source/CsvMetrics/CsvMetrics.Build.cs
@@ -0,0 +1,36 @@
+// Copyright Epic Games, Inc. All Rights Reserved.
+
+using System.IO;
+using UnrealBuildTool;
+
+public class CsvMetrics : ModuleRules
+{
+	public CsvMetrics(ReadOnlyTargetRules Target) : base(Target)
+	{
+		PCHUsage = PCHUsageMode.UseExplicitOrSharedPCHs;
+
+		/**
+		 * Add public dependencies that you statically link with here
+		 *	NOTE: Please use the "PrivateDependencyModuleNames" list to add your module unless
+		 *		it containing "types" that are required in this module's public headers.
+		 */
+
+		PublicDependencyModuleNames.AddRange
+		(
+			new string[]
+			{
+				"CoreUObject",
+				"WorldMetricsCore",
+			}
+		);
+
+		PrivateDependencyModuleNames.AddRange
+		(
+			new string[]
+			{
+				"Core",
+				"Engine",
+			}
+		);
+	}
+}
diff --git a/Engine/Plugins/WorldMetrics/Source/CsvMetrics/Private/CsvActorCountMetric.cpp b/Engine/Plugins/WorldMetrics/Source/CsvMetrics/Private/CsvActorCountMetric.cpp
new file mode 100644
--- /dev/null
+++ b/Engine/Plugins/WorldMetrics/Source/CsvMetrics/Private/CsvActorCountMetric.cpp
@@ -0,0 +1,98 @@
+// Copyright Epic Games, Inc. All Rights Reserved.
+
+#include "CsvActorCountMetric.h"
+
+#include "GameFramework/Actor.h"
+#include "HAL/IConsoleManager.h"
+#include "ProfilingDebugging/CsvProfiler.h"
+#include "WorldMetrics.h"
+#include "WorldMetricsActorTracker.h"
+
+#include UE_INLINE_GENERATED_CPP_BY_NAME(CsvActorCountMetric)
+
+CSV_DECLARE_CATEGORY_MODULE_EXTERN(ENGINE_API, ActorCount);
+
+namespace UE::WorldMetrics::Private
+{
+static TAutoConsoleVariable<int32> CVarRecordActorCountsToCSVThreshold(
+	TEXT("csv.RecordActorCountsThreshold"),
+	5,
+	TEXT("Number of instances of a native Actor class required before recording to a CSV stat"));
+
+}  // namespace UE::WorldMetrics::Private
+
+//---------------------------------------------------------------------------------------------------------------------
+// UCsvActorCountMetric
+//---------------------------------------------------------------------------------------------------------------------
+
+void UCsvActorCountMetric::Initialize()
+{
+	UE::WorldMetrics::AcquireExtension<UWorldMetricsActorTracker>(this);
+}
+
+void UCsvActorCountMetric::Deinitialize()
+{
+	UE::WorldMetrics::ReleaseExtension<UWorldMetricsActorTracker>(this);
+
+	ActorClassNameCounter.Reset();
+	TotalActorCount = 0;
+}
+
+SIZE_T UCsvActorCountMetric::GetAllocatedSize() const
+{
+	return ActorClassNameCounter.GetAllocatedSize();
+}
+
+int32 UCsvActorCountMetric::NumActors(const FName& NativeClassName) const
+{
+	if (const int32* ActorClassCount = ActorClassNameCounter.Find(NativeClassName))
+	{
+		return *ActorClassCount;
+	}
+	return 0;
+}
+
+void UCsvActorCountMetric::Update(float /*DeltaTimeInSeconds*/)
+{
+#if CSV_PROFILER
+	if (FCsvProfiler::Get()->IsCapturing())
+	{
+		const int32 Threshold = UE::WorldMetrics::Private::CVarRecordActorCountsToCSVThreshold.GetValueOnAnyThread();
+		for (const TPair<FName, int32>& ActorClassNameCount : ActorClassNameCounter)
+		{
+			if (ActorClassNameCount.Value > Threshold)
+			{
+				FCsvProfiler::Get()->RecordCustomStat(
+					ActorClassNameCount.Key, CSV_CATEGORY_INDEX(ActorCount), ActorClassNameCount.Value,
+					ECsvCustomStatOp::Set);
+			}
+		}
+	}
+#endif	// CSV_PROFILER
+}
+
+void UCsvActorCountMetric::OnActorAdded(const AActor* Actor)
+{
+	if (!Actor->HasAnyFlags(RF_ArchetypeObject | RF_ClassDefaultObject))
+	{
+		const UClass* ParentNativeClass = GetParentNativeClass(Actor->GetClass());
+		const FName NativeClassName = ParentNativeClass ? ParentNativeClass->GetFName() : NAME_None;
+		int32& ActorClassCount = ActorClassNameCounter.FindOrAdd(NativeClassName);
+		++ActorClassCount;
+		++TotalActorCount;
+	}
+}
+
+void UCsvActorCountMetric::OnActorRemoved(const AActor* Actor)
+{
+	if (!Actor->HasAnyFlags(RF_ArchetypeObject | RF_ClassDefaultObject))
+	{
+		const UClass* ParentNativeClass = GetParentNativeClass(Actor->GetClass());
+		const FName NativeClassName = ParentNativeClass ? ParentNativeClass->GetFName() : NAME_None;
+		if (int32* ActorClassCount = ActorClassNameCounter.Find(NativeClassName))
+		{
+			--(*ActorClassCount);
+		}
+		--TotalActorCount;
+	}
+}
diff --git a/Engine/Plugins/WorldMetrics/Source/CsvMetrics/Private/CsvMetricsModule.cpp b/Engine/Plugins/WorldMetrics/Source/CsvMetrics/Private/CsvMetricsModule.cpp
new file mode 100644
--- /dev/null
+++ b/Engine/Plugins/WorldMetrics/Source/CsvMetrics/Private/CsvMetricsModule.cpp
@@ -0,0 +1,5 @@
+// Copyright Epic Games, Inc. All Rights Reserved.
+
+#include "Modules/ModuleManager.h"
+
+IMPLEMENT_MODULE(FDefaultModuleImpl, CsvMetrics)
diff --git a/Engine/Plugins/WorldMetrics/Source/CsvMetrics/Private/CsvMetricsSubsystem.cpp b/Engine/Plugins/WorldMetrics/Source/CsvMetrics/Private/CsvMetricsSubsystem.cpp
new file mode 100644
--- /dev/null
+++ b/Engine/Plugins/WorldMetrics/Source/CsvMetrics/Private/CsvMetricsSubsystem.cpp
@@ -0,0 +1,93 @@
+// Copyright Epic Games, Inc. All Rights Reserved.
+
+#include "CsvMetricsSubsystem.h"
+
+#include "Engine/World.h"
+#include "ProfilingDebugging/CsvProfiler.h"
+#include "WorldMetrics.h"
+
+#include UE_INLINE_GENERATED_CPP_BY_NAME(CsvMetricsSubsystem)
+
+namespace UE::CsvMetrics::Private
+{
+bool CanHaveCsvMetrics(const UWorld* World)
+{
+#if CSV_PROFILER
+	return World && World->IsGameWorld();
+#else
+	return false;
+#endif	// CSV_PROFILER
+}
+
+}  // namespace UE::CsvMetrics::Private
+
+//---------------------------------------------------------------------------------------------------------------------
+// UCsvMetricsSubsystem
+//---------------------------------------------------------------------------------------------------------------------
+
+bool UCsvMetricsSubsystem::ShouldCreateSubsystem(UObject* Outer) const
+{
+	// TODO: enable once CSVActorClassNameToCountMap is removed.
+	return false; // UE::CsvMetrics::Private::CanHaveCsvMetrics(Cast<UWorld>(Outer));
+}
+
+void UCsvMetricsSubsystem::Initialize(FSubsystemCollectionBase& Collection)
+{
+	BindProfilerCallbacks();
+
+	Super::Initialize(Collection);
+}
+
+void UCsvMetricsSubsystem::Deinitialize()
+{
+	UnbindProfilerCallbacks();
+
+	Super::Deinitialize();
+}
+
+void UCsvMetricsSubsystem::BindProfilerCallbacks()
+{
+#if CSV_PROFILER
+	FCsvProfiler* CsvProfiler = FCsvProfiler::Get();
+	if (ensure(CsvProfiler))
+	{
+		ProfileStartHandle = CsvProfiler->OnCSVProfileStart().AddUObject(this, &UCsvMetricsSubsystem::AddMetrics);
+
+		ProfileEndHandle = CsvProfiler->OnCSVProfileEnd().AddUObject(this, &UCsvMetricsSubsystem::RemoveMetrics);
+	}
+#endif	// CSV_PROFILER
+}
+
+void UCsvMetricsSubsystem::UnbindProfilerCallbacks()
+{
+#if CSV_PROFILER
+	FCsvProfiler* CsvProfiler = FCsvProfiler::Get();
+	if (ensure(CsvProfiler))
+	{
+		CsvProfiler->OnCSVProfileStart().Remove(ProfileStartHandle);
+		CsvProfiler->OnCSVProfileEnd().Remove(ProfileEndHandle);
+	}
+#endif	// CSV_PROFILER
+}
+
+void UCsvMetricsSubsystem::AddMetrics()
+{
+	UWorld* World = GetWorld();
+	check(World);
+
+	for (const TSubclassOf<UWorldMetricInterface>& MetricClass : Metrics)
+	{
+		UE::WorldMetrics::AddMetric(World, MetricClass);
+	}
+}
+
+void UCsvMetricsSubsystem::RemoveMetrics()
+{
+	UWorld* World = GetWorld();
+	check(World);
+
+	for (const TSubclassOf<UWorldMetricInterface>& MetricClass : Metrics)
+	{
+		UE::WorldMetrics::RemoveMetric(World, MetricClass);
+	}
+}
diff --git a/Engine/Plugins/WorldMetrics/Source/CsvMetrics/Public/CsvActorCountMetric.h b/Engine/Plugins/WorldMetrics/Source/CsvMetrics/Public/CsvActorCountMetric.h
new file mode 100644
--- /dev/null
+++ b/Engine/Plugins/WorldMetrics/Source/CsvMetrics/Public/CsvActorCountMetric.h
@@ -0,0 +1,46 @@
+// Copyright Epic Games, Inc. All Rights Reserved.
+
+#pragma once
+
+#include "WorldMetricInterface.h"
+#include "WorldMetricsActorTrackerSubscriber.h"
+
+#include "CsvActorCountMetric.generated.h"
+
+/**
+ * World actor count metric
+ *
+ * This metrics provides a basic actor counter sorting the world's actors into categories.
+ * On update it writes the actor counts to CSV if the CSV profiler is available.
+ */
+UCLASS()
+class UCsvActorCountMetric : public UWorldMetricInterface, public IWorldMetricsActorTrackerSubscriber
+{
+	GENERATED_BODY()
+
+	/** UWorldMetricInterface interface. */
+	[[nodiscard]] CSVMETRICS_API virtual SIZE_T GetAllocatedSize() const override;
+
+	[[nodiscard]] CSVMETRICS_API int32 NumActors() const
+	{
+		return TotalActorCount;
+	}
+
+	[[nodiscard]] CSVMETRICS_API int32 NumActors(const FName& NativeClassName) const;
+
+private:
+	/** Map of Actor class names to instance counts */
+	TMap<FName, int32> ActorClassNameCounter;
+
+	/** Count of total actors created */
+	int32 TotalActorCount = 0;
+
+	/** UWorldMetricInterface interface. */
+	virtual void Initialize() override;
+	virtual void Deinitialize() override;
+	virtual void Update(float DeltaTimeInSeconds) override;
+
+	/** IWorldMetricsActorTrackerSubscriber interface. */
+	virtual void OnActorRemoved(const AActor* Actor) ;
+	virtual void OnActorAdded(const AActor* Actor);
+};
diff --git a/Engine/Plugins/WorldMetrics/Source/CsvMetrics/Public/CsvMetricsSubsystem.h b/Engine/Plugins/WorldMetrics/Source/CsvMetrics/Public/CsvMetricsSubsystem.h
new file mode 100644
--- /dev/null
+++ b/Engine/Plugins/WorldMetrics/Source/CsvMetrics/Public/CsvMetricsSubsystem.h
@@ -0,0 +1,45 @@
+// Copyright Epic Games, Inc. All Rights Reserved.
+
+#pragma once
+
+#include "Subsystems/WorldSubsystem.h"
+
+#include "CsvMetricsSubsystem.generated.h"
+
+class UWorldMetricInterface;
+
+/**
+ * Csv metrics subsystem
+ *
+ * This subsystem registers/unregisters its CSV metric collection whenever a CSV profiler capture is performed.
+ */
+UCLASS(MinimalAPI, Config = Engine, DefaultConfig)
+class UCsvMetricsSubsystem final : public UWorldSubsystem
+{
+	GENERATED_BODY()
+
+public:
+	//~ Begin USubsystem Interface.
+	CSVMETRICS_API virtual bool ShouldCreateSubsystem(UObject* Outer) const override;
+	CSVMETRICS_API virtual void Initialize(FSubsystemCollectionBase& Collection) override;
+	CSVMETRICS_API virtual void Deinitialize() override;
+	//~ End USubsystem Interface.
+
+	/** The CSV metric collection to add/remove when a CSV profiler capture stats/ends. */
+	UPROPERTY(Config)
+	TArray<TSubclassOf<UWorldMetricInterface>> Metrics;
+
+private:
+	FDelegateHandle ProfileStartHandle;
+	FDelegateHandle ProfileEndHandle;
+
+	/**
+	 * Binds/Unbinds handles to OnCSVProfileStart and OnCSVProfileEnd in order register/unregister this system's CSV
+	 * metric collection.
+	 */
+	void BindProfilerCallbacks();
+	void UnbindProfilerCallbacks();
+
+	void AddMetrics();
+	void RemoveMetrics();
+};
diff --git a/Engine/Plugins/WorldMetrics/Source/WorldMetricsCore/Private/WorldMetrics.cpp b/Engine/Plugins/WorldMetrics/Source/WorldMetricsCore/Private/WorldMetrics.cpp
new file mode 100644
--- /dev/null
+++ b/Engine/Plugins/WorldMetrics/Source/WorldMetricsCore/Private/WorldMetrics.cpp
@@ -0,0 +1,153 @@
+// Copyright Epic Games, Inc. All Rights Reserved.
+
+#include "WorldMetrics.h"
+
+#include "Engine/World.h"
+#include "WorldMetricInterface.h"
+#include "WorldMetricsExtension.h"
+#include "WorldMetricsLog.h"
+#include "WorldMetricsSubsystem.h"
+
+//---------------------------------------------------------------------------------------------------------------------
+// UE::WorldMetrics
+//---------------------------------------------------------------------------------------------------------------------
+
+namespace UE::WorldMetrics
+{
+
+namespace Private
+{
+
+/**
+ * World Metrics subsystem inline helper that logs the caller function whenever the world or subsystem are null.
+ *
+ * @return A valid pointer to the world metric subsystem or nullptr otherwise.
+ */
+FORCEINLINE UWorldMetricsSubsystem* GetSubsystem(const UWorld* World)
+{
+	if (UNLIKELY(!World))
+	{
+		UE_LOG(LogWorldMetrics, Error, TEXT("[%hs] Unexpected null World"), __FUNCTION__);
+		return nullptr;
+	}
+
+	UWorldMetricsSubsystem* Subsystem = World->GetSubsystem<UWorldMetricsSubsystem>();
+	if (UNLIKELY(!Subsystem))
+	{
+		UE_LOG(LogWorldMetrics, Error, TEXT("[%hs] Unexpected null World Metrics Subsystem"), __FUNCTION__);
+	}
+	return Subsystem;
+}
+}  // namespace Private
+
+UWorldMetricsSubsystem* GetSubsystem(const UWorld* World)
+{
+	return Private::GetSubsystem(World);
+}
+
+UWorldMetricInterface* GetMetric(const UWorld* World, const TSubclassOf<UWorldMetricInterface>& MetricClass)
+{
+	UWorldMetricsSubsystem* Subsystem = Private::GetSubsystem(World);
+	if (UNLIKELY(!Subsystem))
+	{
+		return nullptr;
+	}
+
+	return Subsystem->GetMetric(MetricClass);
+}
+
+bool AddMetric(const UWorld* World, const TSubclassOf<UWorldMetricInterface>& MetricClass)
+{
+	UWorldMetricsSubsystem* Subsystem = Private::GetSubsystem(World);
+	if (UNLIKELY(!Subsystem))
+	{
+		return false;
+	}
+
+	return Subsystem->AddMetric(MetricClass);
+}
+
+bool RemoveMetric(const UWorld* World, const TSubclassOf<UWorldMetricInterface>& MetricClass)
+{
+	UWorldMetricsSubsystem* Subsystem = Private::GetSubsystem(World);
+	if (UNLIKELY(!Subsystem))
+	{
+		return false;
+	}
+
+	return Subsystem->RemoveMetric(MetricClass);
+}
+
+UWorldMetricsExtension* AcquireExtension(
+	UWorldMetricInterface* InOwner,
+	const TSubclassOf<UWorldMetricsExtension>& InExtensionClass)
+{
+	if (UNLIKELY(!InOwner))
+	{
+		UE_LOG(LogWorldMetrics, Error, TEXT("[%hs] Unexpected invalid owner"), __FUNCTION__);
+		return nullptr;
+	}
+
+	UWorldMetricsSubsystem* Subsystem = Private::GetSubsystem(InOwner->GetWorld());
+	if (UNLIKELY(!Subsystem))
+	{
+		return nullptr;
+	}
+
+	return Subsystem->AcquireExtension(InOwner, InExtensionClass);
+}
+
+UWorldMetricsExtension* AcquireExtension(
+	UWorldMetricsExtension* InOwner,
+	const TSubclassOf<UWorldMetricsExtension>& InExtensionClass)
+{
+	if (UNLIKELY(!InOwner))
+	{
+		UE_LOG(LogWorldMetrics, Error, TEXT("[%hs] Unexpected invalid owner"), __FUNCTION__);
+		return nullptr;
+	}
+
+	UWorldMetricsSubsystem* Subsystem = Private::GetSubsystem(InOwner->GetWorld());
+	if (UNLIKELY(!Subsystem))
+	{
+		return nullptr;
+	}
+
+	return Subsystem->AcquireExtension(InOwner, InExtensionClass);
+}
+
+bool ReleaseExtension(UWorldMetricInterface* InOwner, const TSubclassOf<UWorldMetricsExtension>& InExtensionClass)
+{
+	if (UNLIKELY(!InOwner))
+	{
+		UE_LOG(LogWorldMetrics, Error, TEXT("[%hs] Unexpected invalid owner"), __FUNCTION__);
+		return false;
+	}
+
+	UWorldMetricsSubsystem* Subsystem = Private::GetSubsystem(InOwner->GetWorld());
+	if (UNLIKELY(!Subsystem))
+	{
+		return false;
+	}
+
+	return Subsystem->ReleaseExtension(InOwner, InExtensionClass);
+}
+
+bool ReleaseExtension(UWorldMetricsExtension* InOwner, const TSubclassOf<UWorldMetricsExtension>& InExtensionClass)
+{
+	if (UNLIKELY(!InOwner))
+	{
+		UE_LOG(LogWorldMetrics, Error, TEXT("[%hs] Unexpected invalid owner"), __FUNCTION__);
+		return false;
+	}
+
+	UWorldMetricsSubsystem* Subsystem = Private::GetSubsystem(InOwner->GetWorld());
+	if (UNLIKELY(!Subsystem))
+	{
+		return false;
+	}
+
+	return Subsystem->ReleaseExtension(InOwner, InExtensionClass);
+}
+
+}  // namespace UE::WorldMetrics
diff --git a/Engine/Plugins/WorldMetrics/Source/WorldMetricsCore/Private/WorldMetricsActorTracker.cpp b/Engine/Plugins/WorldMetrics/Source/WorldMetricsCore/Private/WorldMetricsActorTracker.cpp
new file mode 100644
--- /dev/null
+++ b/Engine/Plugins/WorldMetrics/Source/WorldMetricsCore/Private/WorldMetricsActorTracker.cpp
@@ -0,0 +1,118 @@
+// Copyright Epic Games, Inc. All Rights Reserved.
+
+#include "WorldMetricsActorTracker.h"
+
+#include "Engine/Level.h"
+#include "Engine/World.h"
+#include "WorldMetricsActorTrackerSubscriber.h"
+#include "WorldMetricsLog.h"
+
+#include UE_INLINE_GENERATED_CPP_BY_NAME(WorldMetricsActorTracker)
+
+//---------------------------------------------------------------------------------------------------------------------
+// UWorldMetricsActorTracker
+//---------------------------------------------------------------------------------------------------------------------
+
+void UWorldMetricsActorTracker::Initialize()
+{
+	UE_LOG(LogWorldMetrics, Log, TEXT("[%hs]"), __FUNCTION__);
+
+	BindWorldCallbacks();
+}
+
+void UWorldMetricsActorTracker::Deinitialize()
+{
+	UE_LOG(LogWorldMetrics, Log, TEXT("[%hs]"), __FUNCTION__);
+
+	UnbindWorldCallbacks();
+
+	Subscribers.Reset();
+}
+
+void UWorldMetricsActorTracker::BindWorldCallbacks()
+{
+	check(GetWorld());
+
+	ActorAddedToWorldHandle = GetWorld()->AddOnPostRegisterAllActorComponentsHandler(
+		FOnPostRegisterAllActorComponents::FDelegate::CreateLambda([this](AActor* Actor)
+																   { NotifyOnActorAdded(Actor); }));
+
+	ActorRemovedFromWorldHandle = GetWorld()->AddOnPreUnregisterAllActorComponentsHandler(
+		FOnPreUnregisterAllActorComponents::FDelegate::CreateLambda([this](AActor* Actor)
+																	{ NotifyOnActorRemoved(Actor); }));
+}
+
+void UWorldMetricsActorTracker::UnbindWorldCallbacks()
+{
+	check(GetWorld());
+
+	GetWorld()->RemoveOnPostRegisterAllActorComponentsHandler(ActorAddedToWorldHandle);
+	ActorAddedToWorldHandle.Reset();
+
+	GetWorld()->RemoveOnPreUnregisterAllActorComponentsHandler(ActorRemovedFromWorldHandle);
+	ActorRemovedFromWorldHandle.Reset();
+}
+
+SIZE_T UWorldMetricsActorTracker::GetAllocatedSize() const
+{
+	SIZE_T Result = Subscribers.GetAllocatedSize();
+	return Result;
+}
+
+void UWorldMetricsActorTracker::OnAcquire(UObject* InOwner)
+{
+	if (IWorldMetricsActorTrackerSubscriber* Subscriber = Cast<IWorldMetricsActorTrackerSubscriber>(InOwner))
+	{
+		Subscribers.Emplace(Subscriber);
+		NotifyExistingActors(Subscriber);
+	}
+}
+
+void UWorldMetricsActorTracker::OnRelease(UObject* InOwner)
+{
+	if (IWorldMetricsActorTrackerSubscriber* Subscriber = Cast<IWorldMetricsActorTrackerSubscriber>(InOwner))
+	{
+		Subscribers.RemoveSwap(Subscriber);
+	}
+}
+
+void UWorldMetricsActorTracker::NotifyExistingActors(IWorldMetricsActorTrackerSubscriber* Subscriber)
+{
+	TRACE_CPUPROFILER_EVENT_SCOPE(UWorldMetricsActorTracker::NotifyExistingActors);
+
+	check(Subscriber);
+
+	UWorld* World = GetWorld();
+	check(World);
+
+	for (ULevel* Level : World->GetLevels())
+	{
+		for (const AActor* Actor : Level->Actors)
+		{
+			if (Actor && Actor->HasActorRegisteredAllComponents())
+			{
+				Subscriber->OnActorAdded(Actor);
+			}
+		}
+	}
+}
+
+void UWorldMetricsActorTracker::NotifyOnActorAdded(const AActor* Actor)
+{
+	TRACE_CPUPROFILER_EVENT_SCOPE(UWorldMetricsActorTracker::NotifyOnActorAdded);
+
+	for (IWorldMetricsActorTrackerSubscriber* Subscriber : Subscribers)
+	{
+		Subscriber->OnActorAdded(Actor);
+	}
+}
+
+void UWorldMetricsActorTracker::NotifyOnActorRemoved(const AActor* Actor)
+{
+	TRACE_CPUPROFILER_EVENT_SCOPE(UWorldMetricsActorTracker::NotifyOnActorRemoved);
+
+	for (IWorldMetricsActorTrackerSubscriber* Subscriber : Subscribers)
+	{
+		Subscriber->OnActorRemoved(Actor);
+	}
+}
diff --git a/Engine/Plugins/WorldMetrics/Source/WorldMetricsCore/Private/WorldMetricsConsoleCommands.cpp b/Engine/Plugins/WorldMetrics/Source/WorldMetricsCore/Private/WorldMetricsConsoleCommands.cpp
new file mode 100644
--- /dev/null
+++ b/Engine/Plugins/WorldMetrics/Source/WorldMetricsCore/Private/WorldMetricsConsoleCommands.cpp
@@ -0,0 +1,168 @@
+// Copyright Epic Games, Inc. All Rights Reserved.
+
+#include "WorldMetricsConsoleCommands.h"
+
+#if WITH_WORLDMETRICS_DEBUG
+
+#include "HAL/IConsoleManager.h"
+#include "String/LexFromString.h"
+#include "WorldMetricInterface.h"
+#include "WorldMetrics.h"
+#include "WorldMetricsLog.h"
+#include "WorldMetricsSubsystem.h"
+
+namespace UE::WorldMetrics::Private
+{
+
+static FAutoConsoleCommandWithWorldAndArgs CmdWorldMetricsEnable(
+	TEXT("WorldMetrics"),
+	TEXT("Toggles the World Metrics Subsystem unless a boolean parameter is passed."),
+	FConsoleCommandWithWorldAndArgsDelegate::CreateLambda(
+		[](const TArray<FString>& Args, UWorld* World)
+		{
+			if (UWorldMetricsSubsystem* WorldMetrics = GetSubsystem(World))
+			{
+				bool bEnable = !WorldMetrics->IsEnabled();
+				if (Args.Num())
+				{
+					bEnable = Args[0].ToBool();
+				}
+				WorldMetrics->Enable(bEnable);
+			}
+		}),
+	ECVF_Default);
+
+static FAutoConsoleCommandWithWorldAndArgs CmdWorldMetricsClear(
+	TEXT("WorldMetrics.Clear"),
+	TEXT("Clears the World Metrics Subsystem removing all metrics and extensions."),
+	FConsoleCommandWithWorldAndArgsDelegate::CreateLambda(
+		[](const TArray<FString>& Args, UWorld* World)
+		{
+			if (UWorldMetricsSubsystem* WorldMetrics = GetSubsystem(World))
+			{
+				WorldMetrics->Clear();
+			}
+		}),
+	ECVF_Default);
+
+static FAutoConsoleCommandWithWorldAndArgs CmdWorldMetricsSetUpdateRateInSeconds(
+	TEXT("WorldMetrics.SetUpdateRateInSeconds"),
+	TEXT("Resets the World Metrics Subsystem update rate to its default value or the parameter value if specified."),
+	FConsoleCommandWithWorldAndArgsDelegate::CreateLambda(
+		[](const TArray<FString>& Args, UWorld* World)
+		{
+			if (UWorldMetricsSubsystem* WorldMetrics = GetSubsystem(World))
+			{
+				float UpdateRateInSeconds = 0;
+				if (Args.Num() && Args[0].IsNumeric())
+				{
+					UpdateRateInSeconds = FCString::Atof(*Args[0]);
+				}
+				WorldMetrics->SetUpdateRateInSeconds(UpdateRateInSeconds);
+			}
+		}),
+	ECVF_Default);
+
+static void WorldMetricToggleCmdHandler(
+	const TArray<FString>& Args,
+	UWorld* World,
+	FOutputDevice& Ar,
+	TWeakObjectPtr<UClass> MetricClass)
+{
+	UClass* MetricClassPtr = MetricClass.Get();
+	if (MetricClassPtr == nullptr)
+	{
+		Ar.Logf(TEXT("[%hs]: Invalid class."), __FUNCTION__);
+		return;
+	}
+
+	const bool bMetricExists = GetMetric(World, MetricClassPtr) != nullptr;
+	const FName ClassName = MetricClassPtr->GetFName();
+
+	bool bEnable = !bMetricExists;
+	if (Args.IsValidIndex(0))
+	{
+		LexFromString(bEnable, Args[0]);
+	}
+
+	if (bEnable == bMetricExists)
+	{
+		Ar.Log(*WriteToString<256>(
+			ClassName, ANSITEXTVIEW(" is already "), bEnable ? ANSITEXTVIEW("enabled.") : ANSITEXTVIEW("disabled.")));
+		return;
+	}
+
+	if (bEnable && AddMetric(World, MetricClassPtr))
+	{
+		Ar.Log(*WriteToString<128>(ClassName, ANSITEXTVIEW(" added.")));
+		return;
+	}
+
+	if (!bEnable && RemoveMetric(World, MetricClassPtr))
+	{
+		Ar.Log(*WriteToString<128>(ClassName, ANSITEXTVIEW(" removed.")));
+		return;
+	}
+
+	Ar.Log(*WriteToString<256>(
+		ClassName, ANSITEXTVIEW(" couldn't be "), bEnable ? ANSITEXTVIEW("enabled.") : ANSITEXTVIEW("disabled.")));
+}
+
+static IConsoleObject* RegisterMetricToggleConsoleCommand(IConsoleManager& ConsoleManager, UClass* Class)
+{
+	auto NameStringBuilder = WriteToString<256>(ANSITEXTVIEW("WorldMetrics."), Class->GetFName());
+	return ConsoleManager.RegisterConsoleCommand(
+		NameStringBuilder.ToString(), *WriteToString<256>(NameStringBuilder, ANSITEXTVIEW(" [0|1]")),
+		FConsoleCommandWithWorldArgsAndOutputDeviceDelegate::CreateStatic(
+			&WorldMetricToggleCmdHandler, MakeWeakObjectPtr(Class)));
+}
+
+static void RegisterCommonConsoleCommandsForMetric(
+	IConsoleManager& ConsoleManager,
+	UClass* Class,
+	TArray<IConsoleObject*>& OutConsoleObjects)
+{
+	using FConsoleCommandFunc = IConsoleObject* (*)(IConsoleManager& ConsoleManager, UClass* Class);
+	constexpr FConsoleCommandFunc Funcs[] = {
+		&RegisterMetricToggleConsoleCommand,
+	};
+
+	for (FConsoleCommandFunc Func : Funcs)
+	{
+		IConsoleObject* Command = Invoke(Func, ConsoleManager, Class);
+		if (Command != nullptr)
+		{
+			OutConsoleObjects.Add(Command);
+		}
+		else
+		{
+			UE_LOG(
+				LogWorldMetrics, Warning, TEXT("%hs: Failed to register command for class '%s'"), __FUNCTION__,
+				*Class->GetName());
+		}
+	}
+}
+
+}  // namespace UE::WorldMetrics::Private
+
+void UE::WorldMetrics::RegisterCommonMetricConsoleCommands(
+	IConsoleManager& ConsoleManager,
+	TArray<IConsoleObject*>& OutConsoleObjects)
+{
+	TArray<UClass*> WorldMetricClasses;
+	GetDerivedClasses(UWorldMetricInterface::StaticClass(), WorldMetricClasses, true);
+
+	OutConsoleObjects.Reserve(OutConsoleObjects.Num() + WorldMetricClasses.Num());
+
+	for (UClass* Class : WorldMetricClasses)
+	{
+		if (Class == nullptr || Class->HasAnyClassFlags(CLASS_Abstract))
+		{
+			continue;
+		}
+
+		Private::RegisterCommonConsoleCommandsForMetric(ConsoleManager, Class, OutConsoleObjects);
+	}
+}
+
+#endif	// WITH_WORLDMETRICS_DEBUG
diff --git a/Engine/Plugins/WorldMetrics/Source/WorldMetricsCore/Private/WorldMetricsConsoleCommands.h b/Engine/Plugins/WorldMetrics/Source/WorldMetricsCore/Private/WorldMetricsConsoleCommands.h
new file mode 100644
--- /dev/null
+++ b/Engine/Plugins/WorldMetrics/Source/WorldMetricsCore/Private/WorldMetricsConsoleCommands.h
@@ -0,0 +1,27 @@
+// Copyright Epic Games, Inc. All Rights Reserved.
+
+#pragma once
+
+#include "Containers/ContainersFwd.h"
+#include "WorldMetricsDebug.h"
+
+#if WITH_WORLDMETRICS_DEBUG
+
+struct IConsoleManager;
+class IConsoleObject;
+
+namespace UE::WorldMetrics
+{
+
+/**
+ * Registers common console commands for all World Metric classes.
+ *
+ * @param ConsoleManager A reference to the IConsoleManager interface, which the console commands will be registered
+ * with.
+ * @param OutConsoleObjects An array of IConsoleObject pointers to store references to registered console commands.
+ */
+void RegisterCommonMetricConsoleCommands(IConsoleManager& ConsoleManager, TArray<IConsoleObject*>& OutConsoleObjects);
+
+}  // namespace UE::WorldMetrics
+
+#endif	// WITH_WORLDMETRICS_DEBUG
diff --git a/Engine/Plugins/WorldMetrics/Source/WorldMetricsCore/Private/WorldMetricsCoreModule.cpp b/Engine/Plugins/WorldMetrics/Source/WorldMetricsCore/Private/WorldMetricsCoreModule.cpp
new file mode 100644
--- /dev/null
+++ b/Engine/Plugins/WorldMetrics/Source/WorldMetricsCore/Private/WorldMetricsCoreModule.cpp
@@ -0,0 +1,72 @@
+// Copyright Epic Games, Inc. All Rights Reserved.
+
+#include "Modules/ModuleManager.h"
+#include "WorldMetricsDebug.h"
+
+#if WITH_WORLDMETRICS_DEBUG
+#include "HAL/IConsoleManager.h"
+#include "Misc/CoreDelegates.h"
+#include "WorldMetricsConsoleCommands.h"
+#endif	// WITH_WORLDMETRICS_DEBUG
+
+namespace UE::WorldMetrics
+{
+
+class FWorldMetricsCoreModule : public IModuleInterface
+{
+protected:
+#if WITH_WORLDMETRICS_DEBUG
+	//~ Begin IModuleInterface interface
+	virtual void StartupModule() override;
+	virtual void ShutdownModule() override;
+	//~ End IModuleInterface interface
+
+	void OnPostEngineInit();
+
+	void RegisterConsoleCommands();
+	void UnregisterConsoleCommands();
+	TArray<IConsoleObject*> ConsoleObjects;
+#endif	// WITH_WORLDMETRICS_DEBUG
+};
+
+#if WITH_WORLDMETRICS_DEBUG
+
+void FWorldMetricsCoreModule::StartupModule()
+{
+	FCoreDelegates::OnPostEngineInit.AddRaw(this, &FWorldMetricsCoreModule::OnPostEngineInit);
+}
+
+void FWorldMetricsCoreModule::ShutdownModule()
+{
+	UnregisterConsoleCommands();
+	FCoreDelegates::OnPostEngineInit.RemoveAll(this);
+}
+
+void FWorldMetricsCoreModule::OnPostEngineInit()
+{
+	RegisterConsoleCommands();
+}
+
+void FWorldMetricsCoreModule::RegisterConsoleCommands()
+{
+	IConsoleManager& ConsoleManager = IConsoleManager::Get();
+	RegisterCommonMetricConsoleCommands(ConsoleManager, ConsoleObjects);
+}
+
+void FWorldMetricsCoreModule::UnregisterConsoleCommands()
+{
+	IConsoleManager& ConsoleManager = IConsoleManager::Get();
+	for (IConsoleObject* ConsoleObject : ConsoleObjects)
+	{
+		if (ConsoleObject)
+		{
+			ConsoleManager.UnregisterConsoleObject(ConsoleObject);
+		}
+	}
+	ConsoleObjects.Reset();
+}
+#endif	// WITH_WORLDMETRICS_DEBUG
+
+}  // namespace UE::WorldMetrics
+
+IMPLEMENT_MODULE(UE::WorldMetrics::FWorldMetricsCoreModule, WorldMetricsCore)
diff --git a/Engine/Plugins/WorldMetrics/Source/WorldMetricsCore/Private/WorldMetricsLog.cpp b/Engine/Plugins/WorldMetrics/Source/WorldMetricsCore/Private/WorldMetricsLog.cpp
new file mode 100644
--- /dev/null
+++ b/Engine/Plugins/WorldMetrics/Source/WorldMetricsCore/Private/WorldMetricsLog.cpp
@@ -0,0 +1,5 @@
+// Copyright Epic Games, Inc. All Rights Reserved.
+
+#include "WorldMetricsLog.h"
+
+DEFINE_LOG_CATEGORY(LogWorldMetrics);
diff --git a/Engine/Plugins/WorldMetrics/Source/WorldMetricsCore/Private/WorldMetricsSubsystem.cpp b/Engine/Plugins/WorldMetrics/Source/WorldMetricsCore/Private/WorldMetricsSubsystem.cpp
new file mode 100644
--- /dev/null
+++ b/Engine/Plugins/WorldMetrics/Source/WorldMetricsCore/Private/WorldMetricsSubsystem.cpp
@@ -0,0 +1,559 @@
+// Copyright Epic Games, Inc. All Rights Reserved.
+
+#include "WorldMetricsSubsystem.h"
+
+#include "Algo/IndexOf.h"
+#include "Engine/World.h"
+#include "Misc/CoreDelegates.h"
+#include "ProfilingDebugging/CpuProfilerTrace.h"
+#include "WorldMetricInterface.h"
+#include "WorldMetricsExtension.h"
+#include "WorldMetricsLog.h"
+
+#include UE_INLINE_GENERATED_CPP_BY_NAME(WorldMetricsSubsystem)
+
+namespace UE::WorldMetrics::Private
+{
+
+bool CanHaveWorldMetrics(const UWorld* World)
+{
+	return World && World->IsGameWorld();
+}
+
+}  // namespace UE::WorldMetrics::Private
+
+//---------------------------------------------------------------------------------------------------------------------
+// UWorldMetricsSubsystem
+//---------------------------------------------------------------------------------------------------------------------
+
+bool UWorldMetricsSubsystem::ShouldCreateSubsystem(UObject* Outer) const
+{
+	return UE::WorldMetrics::Private::CanHaveWorldMetrics(Cast<UWorld>(Outer));
+}
+
+void UWorldMetricsSubsystem::Initialize(FSubsystemCollectionBase& Collection)
+{
+	UE_LOG(LogWorldMetrics, Log, TEXT("[%hs]"), __FUNCTION__);
+
+	Super::Initialize(Collection);
+}
+
+void UWorldMetricsSubsystem::Deinitialize()
+{
+	UE_LOG(LogWorldMetrics, Log, TEXT("[%hs]"), __FUNCTION__);
+
+	Clear();
+
+	Super::Deinitialize();
+}
+
+void UWorldMetricsSubsystem::BeginDestroy()
+{
+	if (!ensure(!IsEnabled()))
+	{
+		UE_LOG(LogWorldMetrics, Error, TEXT("[%hs] Unexpected subsystem state: IsEnabled"), __FUNCTION__);
+		Clear();
+	}
+	Super::BeginDestroy();
+}
+
+void UWorldMetricsSubsystem::InitializeMetrics()
+{
+	for (UWorldMetricInterface* Metric : Metrics)
+	{
+		check(Metric);
+		Metric->Initialize();
+	}
+	PendingWarmUpFrames = WarmUpFrames;
+	if (!ensure(PendingWarmUpFrames >= 0))
+	{
+		UE_LOG(
+			LogWorldMetrics, Warning, TEXT("[%hs] Invalid warm-up frames value: %d, reset to zero."), __FUNCTION__,
+			PendingWarmUpFrames);
+		PendingWarmUpFrames = 0;
+	}
+}
+
+void UWorldMetricsSubsystem::DeinitializeMetrics()
+{
+	for (UWorldMetricInterface* Metric : Metrics)
+	{
+		check(Metric);
+		Metric->Deinitialize();
+	}
+}
+
+void UWorldMetricsSubsystem::Clear()
+{
+	RemoveAllMetrics();
+
+	Extensions.Reset();
+	IndexedOwners.Reset();
+
+	UE_LOG(LogWorldMetrics, Log, TEXT("[%hs]"), __FUNCTION__);
+}
+
+void UWorldMetricsSubsystem::Enable(bool bEnable)
+{
+	if (bEnable)
+	{
+		if (IsEnabled())
+		{
+			return;
+		}
+
+		if (!GetWorld()->IsGameWorld())
+		{
+			return;
+		}
+
+		if (Metrics.IsEmpty())
+		{
+			return;
+		}
+
+		InitializeMetrics();
+
+		if (!ensure(UpdateRateInSeconds >= 0.f))
+		{
+			UE_LOG(
+				LogWorldMetrics, Warning, TEXT("[%hs] Invalid update rate value: %.02fs, reset to zero."), __FUNCTION__,
+				UpdateRateInSeconds);
+			UpdateRateInSeconds = 0.f;
+		}
+
+		/*
+		 * Regarding CreateLambda.
+		 *
+		 * This subsystem uses an update ticker, although it's not a TickableWorldSubsystem. There are two reasons for
+		 * this design decision:
+		 *
+		 * 1. The subsystem only requires an update ticker for metric updates. Due to the short-lived nature of these,
+		 *    it makes sense to implement a finer-grain control over the ticker enabling/disabling mechanism and prevent
+		 *    incurring an unnecessary cost to the game update thread whenever the subsystem has no metrics.
+		 * 2. At the moment of this writing, the validity checks included in UObject delegates incur a significant cost,
+		 *    negatively impacting the performance of metric updates. This subsystem exclusively owns world metrics and
+		 *    extensions, and their lifetime is bound to the system's Initialize/Deinitialize methods. For this reason,
+		 *    we can assume a raw delegate. Note, however, that UObject methods cannot be bound using CreateRaw; hence,
+		 *    CreateLambda is used instead.
+		 */
+		UpdateTickerHandle = FTSTicker::GetCoreTicker().AddTicker(
+			FTickerDelegate::CreateLambda(
+				[this](float DeltaTimeInSeconds)
+				{
+					OnUpdate(DeltaTimeInSeconds);
+					return true;
+				}),
+			UpdateRateInSeconds);
+	}
+	else if (IsEnabled())
+	{
+		FTSTicker::GetCoreTicker().RemoveTicker(UpdateTickerHandle);
+		UpdateTickerHandle.Reset();
+
+		DeinitializeMetrics();
+	}
+}
+
+void UWorldMetricsSubsystem::SetUpdateRateInSeconds(float InSeconds)
+{
+	if (!ensure(InSeconds >= 0.f))
+	{
+		UE_LOG(
+			LogWorldMetrics, Warning, TEXT("[%hs] Invalid UpdateRateInSeconds input value %.04f: ignored."),
+			__FUNCTION__, InSeconds);
+		return;
+	}
+
+	if (UpdateRateInSeconds != InSeconds)
+	{
+		UpdateRateInSeconds = InSeconds;
+
+		if (IsEnabled())
+		{
+			Enable(false);
+			Enable(true);
+		}
+	}
+}
+
+int32 UWorldMetricsSubsystem::NumExtensions() const
+{
+	return Extensions.Num();
+}
+
+bool UWorldMetricsSubsystem::HasAnyExtension() const
+{
+	return !Extensions.IsEmpty();
+}
+
+int32 UWorldMetricsSubsystem::NumMetrics() const
+{
+	return Metrics.Num();
+}
+
+bool UWorldMetricsSubsystem::HasAnyMetric() const
+{
+	return !Metrics.IsEmpty();
+}
+
+UWorldMetricInterface* UWorldMetricsSubsystem::GetMetric(const TSubclassOf<UWorldMetricInterface>& InMetricClass) const
+{
+	if (InMetricClass->HasAnyClassFlags(CLASS_Abstract))
+	{
+		UE_LOG(
+			LogWorldMetrics, Error, TEXT("[%hs] Parameter metric class is abstract: %s"), __FUNCTION__,
+			*InMetricClass->GetFName().ToString());
+		return nullptr;
+	}
+
+	const int32 MetricIndex = GetMetricIndex(InMetricClass);
+	if (MetricIndex == INDEX_NONE)
+	{
+		return nullptr;
+	}
+
+	return Metrics[MetricIndex];
+}
+
+bool UWorldMetricsSubsystem::AddMetric(const TSubclassOf<UWorldMetricInterface>& InMetricClass)
+{
+	if (InMetricClass->HasAnyClassFlags(CLASS_Abstract))
+	{
+		UE_LOG(
+			LogWorldMetrics, Error, TEXT("[%hs] Parameter metric class is abstract: %s"), __FUNCTION__,
+			*InMetricClass->GetFName().ToString());
+		return false;
+	}
+
+	const int32 MetricIndex = GetMetricIndex(InMetricClass);
+	if (MetricIndex != INDEX_NONE)
+	{
+		return false;
+	}
+
+	UWorldMetricInterface* Metric = NewObject<UWorldMetricInterface>(this, InMetricClass, NAME_None, RF_Transient);
+	if (Metric == nullptr)
+	{
+		UE_LOG(
+			LogWorldMetrics, Error, TEXT("[%hs] Failed to create metric of class: %s"), __FUNCTION__,
+			*InMetricClass->GetFName().ToString());
+		return false;
+	}
+
+	Metrics.Emplace(Metric);
+	if (IsEnabled())
+	{
+		Metric->Initialize();
+	}
+	else
+	{
+		Enable(true);
+	}
+
+	UE_LOG(
+		LogWorldMetrics, Log, TEXT("[%hs] Added metric of class %s."), __FUNCTION__,
+		*InMetricClass->GetFName().ToString());
+
+	return true;
+}
+
+bool UWorldMetricsSubsystem::RemoveMetric(const TSubclassOf<UWorldMetricInterface>& InMetricClass)
+{
+	const int32 MetricIndex = GetMetricIndex(InMetricClass);
+	if (MetricIndex == INDEX_NONE)
+	{
+		return false;
+	}
+
+	UWorldMetricInterface* Metric = Metrics[MetricIndex];
+	if (IsEnabled())
+	{
+		Metric->Deinitialize();
+		VerifyMetricReleasedAllExtensions(Metric);
+		Metric->MarkAsGarbage();
+	}
+	Metrics.RemoveAt(MetricIndex);
+	if (Metrics.IsEmpty())
+	{
+		VerifyRemoveOrphanExtensions();
+		Enable(false);
+	}
+
+	UE_LOG(
+		LogWorldMetrics, Log, TEXT("[%hs] Removed metric of class %s"), __FUNCTION__,
+		*InMetricClass->GetFName().ToString());
+
+	return true;
+}
+
+void UWorldMetricsSubsystem::RemoveAllMetrics()
+{
+	for (int32 MetricIndex = 0; MetricIndex < Metrics.Num(); ++MetricIndex)
+	{
+		UWorldMetricInterface* Metric = Metrics[MetricIndex];
+		if (IsEnabled())
+		{
+			Metric->Deinitialize();
+			VerifyMetricReleasedAllExtensions(Metric);
+			Metric->MarkAsGarbage();
+		}
+	}
+	Metrics.Reset();
+	VerifyRemoveOrphanExtensions();
+	Enable(false);
+}
+
+void UWorldMetricsSubsystem::ForEachMetric(const TFunctionRef<bool(const UWorldMetricInterface*)>& Func) const
+{
+	for (const UWorldMetricInterface* Metric : Metrics)
+	{
+		if (!Func(Metric))
+		{
+			break;
+		}
+	}
+}
+
+void UWorldMetricsSubsystem::ForEachMetric(const TFunctionRef<bool(UWorldMetricInterface*)>& Func)
+{
+	for (UWorldMetricInterface* Metric : Metrics)
+	{
+		if (!Func(Metric))
+		{
+			break;
+		}
+	}
+}
+
+void UWorldMetricsSubsystem::OnUpdate(float DeltaTimeInSeconds)
+{
+	TRACE_CPUPROFILER_EVENT_SCOPE(UWorldMetricsSubsystem::OnUpdate);
+
+	if (PendingWarmUpFrames > 0)
+	{
+		--PendingWarmUpFrames;
+		return;
+	}
+
+	for (UWorldMetricInterface* Metric : Metrics)
+	{
+		check(Metric);
+		Metric->Update(DeltaTimeInSeconds);
+	}
+}
+
+int32 UWorldMetricsSubsystem::GetMetricIndex(const TSubclassOf<UWorldMetricInterface>& InMetricClass) const
+{
+	return Algo::IndexOfByPredicate(
+		Metrics, [InMetricClass](UWorldMetricInterface* Metric) { return Metric->GetClass() == InMetricClass; });
+}
+
+UWorldMetricsExtension* UWorldMetricsSubsystem::AcquireExtension(
+	UWorldMetricInterface* InMetricOwner,
+	const TSubclassOf<UWorldMetricsExtension>& InExtensionClass)
+{
+	return AcquireExtensionInternal(InMetricOwner, InExtensionClass);
+}
+
+UWorldMetricsExtension* UWorldMetricsSubsystem::AcquireExtension(
+	UWorldMetricsExtension* InExtensionOwner,
+	const TSubclassOf<UWorldMetricsExtension>& InExtensionClass)
+{
+	return AcquireExtensionInternal(InExtensionOwner, InExtensionClass);
+}
+
+bool UWorldMetricsSubsystem::ReleaseExtension(
+	UWorldMetricInterface* InMetricOwner,
+	const TSubclassOf<UWorldMetricsExtension>& InExtensionClass)
+{
+	return ReleaseExtensionInternal(InMetricOwner, InExtensionClass);
+}
+
+bool UWorldMetricsSubsystem::ReleaseExtension(
+	UWorldMetricsExtension* InExtensionOwner,
+	const TSubclassOf<UWorldMetricsExtension>& InExtensionClass)
+{
+	return ReleaseExtensionInternal(InExtensionOwner, InExtensionClass);
+}
+
+int32 UWorldMetricsSubsystem::GetExtensionIndex(const TSubclassOf<UWorldMetricsExtension>& InExtensionClass) const
+{
+	return Algo::IndexOfByPredicate(
+		Extensions,
+		[InExtensionClass](UWorldMetricsExtension* Extension) { return Extension->GetClass() == InExtensionClass; });
+}
+
+UWorldMetricsExtension* UWorldMetricsSubsystem::AcquireExtensionInternal(
+	UObject* InOwner,
+	const TSubclassOf<UWorldMetricsExtension>& InExtensionClass)
+{
+	if (UNLIKELY(!InOwner))
+	{
+		UE_LOG(LogWorldMetrics, Error, TEXT("[%hs] Unexpected invalid owner"), __FUNCTION__);
+		return nullptr;
+	}
+
+	if (InExtensionClass->HasAnyClassFlags(CLASS_Abstract))
+	{
+		UE_LOG(
+			LogWorldMetrics, Error, TEXT("[%hs] Parameter extension class is abstract: %s"), __FUNCTION__,
+			*InExtensionClass->GetFName().ToString());
+		return nullptr;
+	}
+
+	UWorldMetricsExtension* Extension = AcquireExistingExtension(InOwner, InExtensionClass);
+	if (!Extension)
+	{
+		Extension = AddExtension(InOwner, InExtensionClass);
+	}
+
+	if (LIKELY(Extension))
+	{
+		Extension->OnAcquire(InOwner);
+	}
+	return Extension;
+}
+
+UWorldMetricsExtension* UWorldMetricsSubsystem::AcquireExistingExtension(
+	UObject* InOwner,
+	const TSubclassOf<UWorldMetricsExtension>& InExtensionClass)
+{
+	const int32 ExtensionIndex = GetExtensionIndex(InExtensionClass);
+	if (ExtensionIndex == INDEX_NONE)
+	{
+		return nullptr;
+	}
+
+	if (!ensure(IndexedOwners.IsValidIndex(ExtensionIndex)))
+	{
+		UE_LOG(LogWorldMetrics, Warning, TEXT("[%hs] Unexpected invalid extension's owner list"), __FUNCTION__);
+		return nullptr;
+	}
+	IndexedOwners[ExtensionIndex].Emplace(InOwner);
+
+	return Extensions[ExtensionIndex];
+}
+
+UWorldMetricsExtension* UWorldMetricsSubsystem::AddExtension(
+	UObject* InOwner,
+	const TSubclassOf<UWorldMetricsExtension>& InExtensionClass)
+{
+	check(Extensions.Num() == IndexedOwners.Num());
+	UWorldMetricsExtension* Extension =
+		NewObject<UWorldMetricsExtension>(this, InExtensionClass, NAME_None, RF_Transient);
+	if (Extension == nullptr)
+	{
+		UE_LOG(
+			LogWorldMetrics, Error, TEXT("[%hs] Failed to create extension of class: %s"), __FUNCTION__,
+			*InExtensionClass->GetFName().ToString());
+		return nullptr;
+	}
+
+	Extensions.Emplace(Extension);
+	IndexedOwners.Emplace_GetRef().Emplace(InOwner);
+	Extension->Initialize();
+
+	UE_LOG(
+		LogWorldMetrics, Log, TEXT("[%hs] Added extension of class: %s"), __FUNCTION__,
+		*InExtensionClass->GetFName().ToString());
+
+	return Extension;
+}
+
+bool UWorldMetricsSubsystem::ReleaseExtensionInternal(
+	UObject* InOwner,
+	const TSubclassOf<UWorldMetricsExtension>& InExtensionClass)
+{
+	if (UNLIKELY(!InOwner))
+	{
+		return false;
+	}
+
+	const int32 ExtensionIndex = GetExtensionIndex(InExtensionClass);
+	if (ExtensionIndex == INDEX_NONE)
+	{
+		return false;
+	}
+
+	if (UNLIKELY(!IndexedOwners[ExtensionIndex].Remove(InOwner)))
+	{
+		UE_LOG(
+			LogWorldMetrics, Error, TEXT("[%hs] Parameter object doesn't own this extension: %s"), __FUNCTION__,
+			*InExtensionClass->GetFName().ToString());
+	}
+	Extensions[ExtensionIndex]->OnRelease(InOwner);
+
+	TryRemoveExtensionAt(ExtensionIndex);
+
+	return true;
+}
+
+bool UWorldMetricsSubsystem::TryRemoveExtensionAt(int32 ExtensionIndex)
+{
+	if (IndexedOwners[ExtensionIndex].IsEmpty())
+	{
+		UWorldMetricsExtension* Extension = Extensions[ExtensionIndex];
+		Extension->Deinitialize();
+		Extension->MarkAsGarbage();
+
+		Extensions.RemoveAtSwap(ExtensionIndex);
+		IndexedOwners.RemoveAtSwap(ExtensionIndex);
+		UE_LOG(
+			LogWorldMetrics, Log, TEXT("[%hs] Removed extension of class %s"), __FUNCTION__,
+			*Extension->GetClass()->GetFName().ToString());
+
+		return true;
+	}
+	return false;
+}
+
+void UWorldMetricsSubsystem::VerifyMetricReleasedAllExtensions(UWorldMetricInterface* InMetric)
+{
+	check(Extensions.Num() == IndexedOwners.Num());
+
+	TArray<int32, TInlineAllocator<DefaultExtensionCapacity>> StaleExtensionIndices;
+	for (int32 ExtensionIndex = 0; ExtensionIndex < Extensions.Num(); ++ExtensionIndex)
+	{
+		if (UNLIKELY(IndexedOwners[ExtensionIndex].Remove(InMetric)))
+		{
+			StaleExtensionIndices.Emplace(ExtensionIndex);
+			UWorldMetricsExtension* Extension = Extensions[ExtensionIndex];
+			UE_LOG(
+				LogWorldMetrics, Warning, TEXT("[%hs] World metric %s did not release extension %s in Deinitialize"),
+				__FUNCTION__, *InMetric->GetFName().ToString(), *Extension->GetName());
+		}
+	}
+
+	// Reverse iterate stale extension indices in case the extension can be removed
+	for (int32 i = StaleExtensionIndices.Num(); i-- > 0;)
+	{
+		TryRemoveExtensionAt(StaleExtensionIndices[i]);
+	}
+}
+
+void UWorldMetricsSubsystem::VerifyRemoveOrphanExtensions()
+{
+	check(Extensions.Num() == IndexedOwners.Num());
+
+	if (Metrics.IsEmpty())
+	{
+		if (UNLIKELY(!Extensions.IsEmpty()))
+		{
+			for (UWorldMetricsExtension* Extension : Extensions)
+			{
+				check(Extension);
+
+				UE_LOG(
+					LogWorldMetrics, Warning,
+					TEXT("[%hs] World extension %s was acquired by another extension but never released."),
+					__FUNCTION__, *Extension->GetName());
+
+				Extension->Deinitialize();
+				Extension->MarkAsGarbage();
+			}
+			Extensions.Reset();
+			IndexedOwners.Reset();
+		}
+	}
+}
diff --git a/Engine/Plugins/WorldMetrics/Source/WorldMetricsCore/Public/WorldMetricInterface.h b/Engine/Plugins/WorldMetrics/Source/WorldMetricsCore/Public/WorldMetricInterface.h
new file mode 100644
--- /dev/null
+++ b/Engine/Plugins/WorldMetrics/Source/WorldMetricsCore/Public/WorldMetricInterface.h
@@ -0,0 +1,49 @@
+// Copyright Epic Games, Inc. All Rights Reserved.
+
+#pragma once
+
+#include "UObject/Object.h"
+
+#include "WorldMetricInterface.generated.h"
+
+/**
+ * World metric's interface
+ *
+ * This is the required interface class to implement world metrics.
+ */
+UCLASS(abstract, MinimalAPI)
+class UWorldMetricInterface : public UObject
+{
+	GENERATED_BODY()
+
+public:
+	/*
+	 * Returns the amount of memory allocated by this class, not including sizeof(*this).
+	 * @return the result size in bytes.
+	 */
+	[[nodiscard]] WORLDMETRICSCORE_API virtual SIZE_T GetAllocatedSize() const
+		PURE_VIRTUAL(GetAllocatedSize, return 0;);
+
+private:
+	/**
+	 * Metric initialization method called by the World Metrics subsystem on metric addition.
+	 */
+	virtual void Initialize()
+	{
+	}
+
+	/**
+	 * Metric initialization method called by the World Metrics subsystem before metric removal.
+	 */
+	virtual void Deinitialize()
+	{
+	}
+
+	/**
+	 * The metric's update method called by the World Metrics subsystem when enabled.
+	 * @param DeltaTimeInSeconds the delta time in seconds from the last update.
+	 */
+	virtual void Update(float DeltaTimeInSeconds) PURE_VIRTUAL(Update, );
+
+	friend class UWorldMetricsSubsystem;
+};
diff --git a/Engine/Plugins/WorldMetrics/Source/WorldMetricsCore/Public/WorldMetrics.h b/Engine/Plugins/WorldMetrics/Source/WorldMetricsCore/Public/WorldMetrics.h
new file mode 100644
--- /dev/null
+++ b/Engine/Plugins/WorldMetrics/Source/WorldMetricsCore/Public/WorldMetrics.h
@@ -0,0 +1,116 @@
+// Copyright Epic Games, Inc. All Rights Reserved.
+
+#pragma once
+
+#include "Templates/SubclassOf.h"
+
+class UWorld;
+class UWorldMetricInterface;
+class UWorldMetricsExtension;
+class UWorldMetricsSubsystem;
+
+namespace UE::WorldMetrics
+{
+
+/**
+ * Gets the World Metrics subsystem. This method expects the parameter world to have the subsystem and will log an error
+ * if it isn't available.
+ *
+ * @return A valid pointer to the world metric subsystem or nullptr otherwise.
+ */
+[[nodiscard]] WORLDMETRICSCORE_API UWorldMetricsSubsystem* GetSubsystem(const UWorld* World);
+
+/**
+ * Gets a world metric from the parameter world's World Metrics subsystem using the parameter metric
+ * class.
+ *
+ * @param World: the world owning the destination World Metrics subsystem.
+ * @param MetricClass: the class of the world metric to be retrieved.
+ * @return A valid pointer to the world metric if found, and nullptr otherwise.
+ */
+[[nodiscard]] WORLDMETRICSCORE_API UWorldMetricInterface* GetMetric(
+	const UWorld* World,
+	const TSubclassOf<UWorldMetricInterface>& MetricClass);
+
+template <typename MetricClass UE_REQUIRES(std::is_base_of_v<UWorldMetricInterface, MetricClass>)>
+[[nodiscard]] MetricClass* GetMetric(const UWorld* World)
+{
+	return static_cast<MetricClass*>(GetMetric(World, MetricClass::StaticClass()));
+}
+
+/**
+ * Adds a new world metric to the parameter world's World Metrics subsystem with the parameter metric class.
+ * A new world metric is instanced only if a metric with the same class doesn't already exist. The destination world
+ * metrics subsystem becomes the owner of the created metric.
+ *
+ * @param World: the world owning the destination World Metrics subsystem.
+ * @param MetricClass: the class of the world metric to be added.
+ * @return true if a new metric is created and false otherwise.
+ */
+WORLDMETRICSCORE_API bool AddMetric(const UWorld* World, const TSubclassOf<UWorldMetricInterface>& MetricClass);
+
+template <typename MetricClass UE_REQUIRES(std::is_base_of_v<UWorldMetricInterface, MetricClass>)>
+bool AddMetric(const UWorld* World)
+{
+	return AddMetric(World, MetricClass::StaticClass());
+}
+
+/**
+ * Removes a new world metric to the parameter world's World Metrics subsystem with the parameter metric class.
+ * A world metric can only be removed if a metric of the class type exists.
+ *
+ * @param World: the world owning the destination World Metrics subsystem.
+ * @param MetricClass: the class of the world metric to remove.
+ * @return true if an existing metric is removed and false otherwise.
+ */
+WORLDMETRICSCORE_API bool RemoveMetric(const UWorld* World, const TSubclassOf<UWorldMetricInterface>& MetricClass);
+
+template <typename MetricClass UE_REQUIRES(std::is_base_of_v<UWorldMetricInterface, MetricClass>)>
+bool RemoveMetric(const UWorld* World)
+{
+	return RemoveMetric(World, MetricClass::StaticClass());
+}
+
+/**
+ * Acquires an extension on behalf of a object. The extension will be created the first time it is acquired,
+ * and conversely it will be destroyed once released by all owners.
+ *
+ * @param InOwner: The object acquiring the extension. Must be a UWorldMetricInterface or UWorldMetricExtension.
+ * @param InExtensionClass: The class type of the extension.
+ * @return A pointer to the extension. The lifetime of the extension is managed by the subsystem.
+ */
+WORLDMETRICSCORE_API UWorldMetricsExtension* AcquireExtension(
+	UWorldMetricInterface* InOwner,
+	const TSubclassOf<UWorldMetricsExtension>& InExtensionClass);
+
+WORLDMETRICSCORE_API UWorldMetricsExtension* AcquireExtension(
+	UWorldMetricsExtension* InOwner,
+	const TSubclassOf<UWorldMetricsExtension>& InExtensionClass);
+
+template <class ExtensionClass, class OwnerType UE_REQUIRES(std::is_base_of_v<UWorldMetricsExtension, ExtensionClass>)>
+ExtensionClass* AcquireExtension(OwnerType* InOwner)
+{
+	return static_cast<ExtensionClass*>(AcquireExtension(InOwner, ExtensionClass::StaticClass()));
+}
+
+/**
+ * Releases a metric's ownership of an extension. If this was the last reference the extension will be disabled.
+ *
+ * @param InOwner: The object whose ownership should be released.
+ * @param InExtensionClass: The class type of the extension.
+ */
+WORLDMETRICSCORE_API bool ReleaseExtension(
+	UWorldMetricInterface* InOwner,
+	const TSubclassOf<UWorldMetricsExtension>& InExtensionClass);
+
+WORLDMETRICSCORE_API bool ReleaseExtension(
+	UWorldMetricsExtension* InOwner,
+	const TSubclassOf<UWorldMetricsExtension>& InExtensionClass);
+
+template <class ExtensionClass, class OwnerType UE_REQUIRES(std::is_base_of_v<UWorldMetricsExtension, ExtensionClass>)>
+bool ReleaseExtension(OwnerType* InOwner)
+{
+	return ReleaseExtension(InOwner, ExtensionClass::StaticClass());
+}
+
+}  // namespace UE::WorldMetrics
diff --git a/Engine/Plugins/WorldMetrics/Source/WorldMetricsCore/Public/WorldMetricsActorTracker.h b/Engine/Plugins/WorldMetrics/Source/WorldMetricsCore/Public/WorldMetricsActorTracker.h
new file mode 100644
--- /dev/null
+++ b/Engine/Plugins/WorldMetrics/Source/WorldMetricsCore/Public/WorldMetricsActorTracker.h
@@ -0,0 +1,72 @@
+// Copyright Epic Games, Inc. All Rights Reserved.
+
+#pragma once
+
+#include "WorldMetricsExtension.h"
+
+#include "WorldMetricsActorTracker.generated.h"
+
+class IWorldMetricsActorTrackerSubscriber;
+
+/**
+ * Actor tracker class
+ *
+ * This class provides a tracking context for actors added and removed from the world. For an actor to be considered
+ * added to the world, it must have all its components registered. This class provides two ways of usage:
+ * - Polling mode: accessible through ForEachActorInWorld.
+ * - Subscriber mode: this mode is enabled through acquire/release semantics. When acquired, the actor tracker
+ * subscribes to the world's PostRegisterAllActorComponents and PreUnregisterAllActorComponentsRegistered events
+ * tracking all unique actors. In addition, metrics implementing IWorldMetricsActorTrackerSubscriber receive actor
+ * addition and removal notifications.
+ */
+UCLASS(MinimalAPI)
+class UWorldMetricsActorTracker final : public UWorldMetricsExtension
+{
+	GENERATED_BODY()
+
+public:
+	/** UWorldMetricsExtension */
+	[[nodiscard]] WORLDMETRICSCORE_API virtual SIZE_T GetAllocatedSize() const override;
+
+private:
+	FDelegateHandle ActorAddedToWorldHandle;
+	FDelegateHandle ActorRemovedFromWorldHandle;
+
+	/** Subscriber list. This list uses raw IWorldMetricsActorTrackerSubscriber pointers are used by design instead of
+	 * TWeakInterfacePtr for performance purposes. This decision is possible because the UWorldMetricsSubsystem is the
+	 * owner of both this extension and all its subscribers, guaranteeing the lifetime requirements. 
+	 */
+	TArray<IWorldMetricsActorTrackerSubscriber*, TInlineAllocator<16>> Subscribers;
+	bool bIsEnabled = false;
+
+	/** UWorldMetricsExtension */
+	virtual void Initialize() override;
+	virtual void Deinitialize() override;
+	virtual void OnAcquire(UObject* InOwner) override;
+	virtual void OnRelease(UObject* InOwner) override;
+
+	/**
+	 * Binds/Unbinds handles to PostRegisterAllActorComponents and PreUnregisterAllActorComponents in order to account
+	 * for actors being added/removed to/from the world.
+	 */
+	void BindWorldCallbacks();
+	void UnbindWorldCallbacks();
+
+	/*
+	 * Notifies the parameter subscriber the addition of all actors in the world with all their components registered.
+	 * @param Subscriber: the subscriber to notify.
+	 */
+	void NotifyExistingActors(IWorldMetricsActorTrackerSubscriber* Subscriber);
+
+	/**
+	 * Fires an addition notification event to all subscribers implementing IWorldMetricsActorTrackerSubscriber.
+	 * @param the actor to notify.
+	 */
+	void NotifyOnActorAdded(const AActor* Actor);
+
+	/**
+	 * Fires a removal notification event to all subscribers implementing IWorldMetricsActorTrackerSubscriber.
+	 * @param the actor to notify.
+	 */
+	void NotifyOnActorRemoved(const AActor* Actor);
+};
diff --git a/Engine/Plugins/WorldMetrics/Source/WorldMetricsCore/Public/WorldMetricsActorTrackerSubscriber.h b/Engine/Plugins/WorldMetrics/Source/WorldMetricsCore/Public/WorldMetricsActorTrackerSubscriber.h
new file mode 100644
--- /dev/null
+++ b/Engine/Plugins/WorldMetrics/Source/WorldMetricsCore/Public/WorldMetricsActorTrackerSubscriber.h
@@ -0,0 +1,37 @@
+// Copyright Epic Games, Inc. All Rights Reserved.
+
+#pragma once
+
+#include "UObject/Interface.h"
+
+#include "WorldMetricsActorTrackerSubscriber.generated.h"
+
+UINTERFACE(MinimalAPI)
+class UWorldMetricsActorTrackerSubscriber : public UInterface
+{
+	GENERATED_BODY()
+};
+
+/**
+ * Actor tracker subscriber interface
+ *
+ * This interface enables the implementing object class to subscribe to UWorldMetricsActorTracker and receive
+ * notifications whenever an actor is added and removed from the world.
+ */
+class IWorldMetricsActorTrackerSubscriber
+{
+	GENERATED_BODY()
+
+public:
+	/**
+	 * Triggers whenever a tracked actor is added.
+	 * @param Actor: the actor pointer which is guaranteed to be valid until its corresponding OnActorRemoved event.
+	 */
+	virtual void OnActorAdded(const AActor* Actor) = 0;
+
+	/**
+	 * Triggers whenever a tracked actor is about to be removed.
+	 * @param Actor: the actor pointer will no longer be valid after this event.
+	 */
+	virtual void OnActorRemoved(const AActor* Actor) = 0;
+};
diff --git a/Engine/Plugins/WorldMetrics/Source/WorldMetricsCore/Public/WorldMetricsDebug.h b/Engine/Plugins/WorldMetrics/Source/WorldMetricsCore/Public/WorldMetricsDebug.h
new file mode 100644
--- /dev/null
+++ b/Engine/Plugins/WorldMetrics/Source/WorldMetricsCore/Public/WorldMetricsDebug.h
@@ -0,0 +1,13 @@
+// Copyright Epic Games, Inc. All Rights Reserved.
+
+#pragma once
+
+#ifndef WITH_WORLDMETRICS_DEBUG
+#define WITH_WORLDMETRICS_DEBUG (!UE_BUILD_SHIPPING && ALLOW_CONSOLE)
+#endif	// WITH_WORLDMETRICS_DEBUG
+
+#if WITH_WORLDMETRICS_DEBUG
+#define WORLDMETRICS_DEBUG_ONLY(...) __VA_ARGS__
+#else
+#define WORLDMETRICS_DEBUG_ONLY(...)
+#endif	// WITH_WORLDMETRICS_DEBUG
diff --git a/Engine/Plugins/WorldMetrics/Source/WorldMetricsCore/Public/WorldMetricsExtension.h b/Engine/Plugins/WorldMetrics/Source/WorldMetricsCore/Public/WorldMetricsExtension.h
new file mode 100644
--- /dev/null
+++ b/Engine/Plugins/WorldMetrics/Source/WorldMetricsCore/Public/WorldMetricsExtension.h
@@ -0,0 +1,58 @@
+// Copyright Epic Games, Inc. All Rights Reserved.
+
+#pragma once
+
+#include "UObject/Object.h"
+
+#include "WorldMetricsExtension.generated.h"
+
+/**
+ * Base class for an extension for the World Metrics subsystem.
+ *
+ * World Metrics subsystem extensions provide custom data/functionality to world metrics and other extensions (one to
+ * many) and have the following characteristics:
+ * 1. Unique: there can only be a single instance of each extension class.
+ * 2. Exclusively managed and owned by the  World Metrics subsystem.. These should never be created directly.
+ * 3. Implement acquired/release semantics. The World Metrics subsystem automatically initializes extensions on
+ * acquisition and deinitializes them on release. The subsystem may deallocate an extension that have no acquisitions.
+ */
+UCLASS(abstract, MinimalAPI)
+class UWorldMetricsExtension : public UObject
+{
+	GENERATED_BODY()
+
+public:
+	/*
+	 * Returns the amount of memory allocated by this class, not including sizeof(*this).
+	 * @return the result size in bytes.
+	 */
+	[[nodiscard]] WORLDMETRICSCORE_API virtual SIZE_T GetAllocatedSize() const
+		PURE_VIRTUAL(GetAllocatedSize, return 0;);
+
+private:
+	/**
+	 * Called by the World Metrics subsystem whenever a new extension is added.
+	 */
+	virtual void Initialize()
+	{
+	}
+
+	/**
+	 * Called by the World Metrics subsystem whenever an extension is about to be removed.
+	 */
+	virtual void Deinitialize()
+	{
+	}
+
+	/* Called by the World Metrics subsystem whenever an extension is acquired. */
+	virtual void OnAcquire(UObject* InOwner)
+	{
+	}
+
+	/* Called by the World Metrics subsystem whenever an extension is released. */
+	virtual void OnRelease(UObject* InOwner)
+	{
+	}
+
+	friend class UWorldMetricsSubsystem;
+};
diff --git a/Engine/Plugins/WorldMetrics/Source/WorldMetricsCore/Public/WorldMetricsLog.h b/Engine/Plugins/WorldMetrics/Source/WorldMetricsCore/Public/WorldMetricsLog.h
new file mode 100644
--- /dev/null
+++ b/Engine/Plugins/WorldMetrics/Source/WorldMetricsCore/Public/WorldMetricsLog.h
@@ -0,0 +1,7 @@
+// Copyright Epic Games, Inc. All Rights Reserved.
+
+#pragma once
+
+#include "Logging/LogMacros.h"
+
+WORLDMETRICSCORE_API DECLARE_LOG_CATEGORY_EXTERN(LogWorldMetrics, Log, All);
diff --git a/Engine/Plugins/WorldMetrics/Source/WorldMetricsCore/Public/WorldMetricsSubsystem.h b/Engine/Plugins/WorldMetrics/Source/WorldMetricsCore/Public/WorldMetricsSubsystem.h
new file mode 100644
--- /dev/null
+++ b/Engine/Plugins/WorldMetrics/Source/WorldMetricsCore/Public/WorldMetricsSubsystem.h
@@ -0,0 +1,271 @@
+// Copyright Epic Games, Inc. All Rights Reserved.
+
+#pragma once
+
+#include "Containers/Ticker.h"
+#include "Subsystems/WorldSubsystem.h"
+
+#include "WorldMetricsSubsystem.generated.h"
+
+class UWorldMetricInterface;
+class UWorldMetricsExtension;
+class UWorldMetricsSubsystem;
+
+/**
+ * World metrics subsystem
+ *
+ * This subsystem provides an interface to add and remove world metrics implementing the UWorldMetricInterface class.
+ * Both the subsystem and individual metrics can be enabled on demand. The World Metrics subsystem is the owner of all
+ * the added metrics and these get automatically garbage collected after removal unless another system holds a
+ * hard-reference to them.
+ */
+UCLASS(MinimalAPI, Config = Engine, DefaultConfig)
+class UWorldMetricsSubsystem final : public UWorldSubsystem
+{
+	GENERATED_BODY()
+
+public:
+	//~ Begin USubsystem
+	WORLDMETRICSCORE_API virtual bool ShouldCreateSubsystem(UObject* Outer) const override;
+	WORLDMETRICSCORE_API virtual void Initialize(FSubsystemCollectionBase& Collection) override;
+	WORLDMETRICSCORE_API virtual void Deinitialize() override;
+	//~ End USubsystem
+
+	//~ Begin UObject
+	WORLDMETRICSCORE_API virtual void BeginDestroy() override;
+	//~ End UObject
+
+	/**
+	 * Enables or disables the subsystem. When enabled the subsystem uses an update ticker to update each of the added
+	 * world metrics. All metrics are automatically initialized when the system is enabled and deinitialized when is
+	 * disabled. The subsystem is automatically disabled on deinitialization.
+	 *
+	 * @param bEnable: true for enabling the system and false otherwise.
+	 */
+	WORLDMETRICSCORE_API void Enable(bool bEnable);
+
+	[[nodiscard]] bool IsEnabled() const
+	{
+		return UpdateTickerHandle.IsValid();
+	}
+
+	WORLDMETRICSCORE_API void Clear();
+
+	/**
+	 * Sets the subsystem update ticker rate in seconds. This method automatically restarts the subsystem if it's
+	 * already enabled. Changing the update rate value on an enabled subsystem causes all metrics to be reinitialized.
+	 *
+	 * @requirement: parameter InSeconds value must be equal or greater than zero.
+	 * @param InSeconds: the update rate value in seconds.
+	 */
+	WORLDMETRICSCORE_API void SetUpdateRateInSeconds(float InSeconds);
+
+	/**
+	 * @return the number of metrics currently contained by the subsystem.
+	 */
+	WORLDMETRICSCORE_API int32 NumMetrics() const;
+
+	/**
+	 * @return true if the subsystem contains any metric.
+	 */
+	WORLDMETRICSCORE_API bool HasAnyMetric() const;
+
+	/**
+	 * @return the number of extensions currently contained by the subsystem.
+	 */
+	WORLDMETRICSCORE_API int32 NumExtensions() const;
+
+	/**
+	 * @return true if the subsystem contains any extension.
+	 */
+	WORLDMETRICSCORE_API bool HasAnyExtension() const;
+
+	/**
+	 * Gets a world metric instance from the subsystem using the parameter metric
+	 * class.
+	 *
+	 * @param InMetricClass: the class of the world metric to be retrieved.
+	 * @return A valid pointer to the world metric if found, and nullptr otherwise.
+	 */
+	[[nodiscard]] WORLDMETRICSCORE_API UWorldMetricInterface* GetMetric(
+		const TSubclassOf<UWorldMetricInterface>& InMetricClass) const;
+
+	template <typename MetricClass UE_REQUIRES(std::is_base_of_v<UWorldMetricInterface, MetricClass>)>
+	[[nodiscard]] MetricClass* GetMetric() const
+	{
+		return static_cast<MetricClass*>(GetMetric(MetricClass::StaticClass()));
+	}
+
+	/**
+	 * Adds a metric instance of the parameter class to the subsystem unless a metric of the same class and/or same ID
+	 * already exists. The system becomes the owner of all added metrics.
+	 *
+	 * @param InMetricClass: the class of the metric to add.
+	 * @param InMetricId: the ID of the metric to add.
+	 * @return true if a metric of the parameter metric class is added and false otherwise.
+	 */
+	WORLDMETRICSCORE_API bool AddMetric(const TSubclassOf<UWorldMetricInterface>& InMetricClass);
+
+	template <typename MetricClass UE_REQUIRES(std::is_base_of_v<UWorldMetricInterface, MetricClass>)>
+	bool AddMetric()
+	{
+		return AddMetric(MetricClass::StaticClass());
+	}
+
+	/**
+	 * Removes the metric matching the parameter class and/or parameter ID from the subsystem.
+	 * The metric object gets automatically garbage-collected unless other system holds a hard-reference to it.
+	 *
+	 * @param InMetricClass: the class corresponding to the world metric to remove.
+	 * @param InMetricId: the ID corresponding to the world metric to remove.
+	 * @return true if a metric matching the parameter class is removed and false otherwise.
+	 */
+	WORLDMETRICSCORE_API bool RemoveMetric(const TSubclassOf<UWorldMetricInterface>& InMetricClass);
+
+	template <typename MetricClass UE_REQUIRES(std::is_base_of_v<UWorldMetricInterface, MetricClass>)>
+	bool RemoveMetric()
+	{
+		return RemoveMetric(MetricClass::StaticClass());
+	}
+
+	WORLDMETRICSCORE_API void RemoveAllMetrics();
+
+	/**
+	 * Invokes the parameter function on each of the metrics contained by the subsystem.
+	 * @param Func The function which will be invoked for each metric. The function should return true to continue
+	 * execution, or false otherwise.
+	 */
+	WORLDMETRICSCORE_API void ForEachMetric(const TFunctionRef<bool(const UWorldMetricInterface*)>& Func) const;
+	WORLDMETRICSCORE_API void ForEachMetric(const TFunctionRef<bool(UWorldMetricInterface*)>& Func);
+
+	/**
+	 * Acquires an extension on behalf of a object. The extension will be created the first time it is acquired,
+	 * and conversely it will be disabled once released by all owners.
+	 *
+	 * @param InOwner: The object acquiring the extension.
+	 * @param InExtensionClass: The class of the desired extension.
+	 * @return A pointer to the extension. The lifetime of the extension is managed by the subsystem.
+	 */
+	WORLDMETRICSCORE_API UWorldMetricsExtension* AcquireExtension(
+		UWorldMetricInterface* InOwner,
+		const TSubclassOf<UWorldMetricsExtension>& InExtensionClass);
+
+	template <class ExtensionClass UE_REQUIRES(std::is_base_of_v<UWorldMetricsExtension, ExtensionClass>)>
+	ExtensionClass* AcquireExtension(UWorldMetricInterface* InMetricOwner)
+	{
+		return static_cast<ExtensionClass*>(AcquireExtension(InMetricOwner, ExtensionClass::StaticClass()));
+	}
+
+	WORLDMETRICSCORE_API UWorldMetricsExtension* AcquireExtension(
+		UWorldMetricsExtension* InOwner,
+		const TSubclassOf<UWorldMetricsExtension>& InExtensionClass);
+
+	template <class ExtensionClass UE_REQUIRES(std::is_base_of_v<UWorldMetricsExtension, ExtensionClass>)>
+	ExtensionClass* AcquireExtension(UWorldMetricsExtension* InExtensionOwner)
+	{
+		return static_cast<ExtensionClass*>(AcquireExtension(InExtensionOwner, ExtensionClass::StaticClass()));
+	}
+
+	/**
+	 * Releases a metric's ownership of an extension. If this was the last reference the extension will be disabled.
+	 *
+	 * @param InOwner: The object whose ownership should be released.
+	 * @param InExtensionClass: The class type of the extension.
+	 */
+	WORLDMETRICSCORE_API bool ReleaseExtension(
+		UWorldMetricInterface* InOwner,
+		const TSubclassOf<UWorldMetricsExtension>& InExtensionClass);
+
+	WORLDMETRICSCORE_API bool ReleaseExtension(
+		UWorldMetricsExtension* InOwner,
+		const TSubclassOf<UWorldMetricsExtension>& InExtensionClass);
+
+	template <class ExtensionClass UE_REQUIRES(std::is_base_of_v<UWorldMetricsExtension, ExtensionClass>)>
+	bool ReleaseExtension(UWorldMetricInterface* InOwner)
+	{
+		return ReleaseExtension(InOwner, ExtensionClass::StaticClass());
+	}
+
+	template <class ExtensionClass UE_REQUIRES(std::is_base_of_v<UWorldMetricsExtension, ExtensionClass>)>
+	bool ReleaseExtension(UWorldMetricsExtension* InExtensionOwner)
+	{
+		return ReleaseExtension(InExtensionOwner, ExtensionClass::StaticClass());
+	}
+
+private:
+	static constexpr int32 DefaultExtensionCapacity = 16;
+	static constexpr int32 DefaultOwnerListCapacity = 32;
+
+	FTSTicker::FDelegateHandle UpdateTickerHandle;
+
+	UPROPERTY(Transient)
+	TArray<TObjectPtr<UWorldMetricInterface>> Metrics;
+
+	UPROPERTY(Transient)
+	TArray<TObjectPtr<UWorldMetricsExtension>> Extensions;
+
+	using FOwnerList = TSet<UObject*, DefaultKeyFuncs<UObject*>, TFixedSetAllocator<DefaultOwnerListCapacity>>;
+	TArray<FOwnerList, TInlineAllocator<DefaultExtensionCapacity>> IndexedOwners;
+
+public:
+	UPROPERTY(Config)
+	float UpdateRateInSeconds = 0.f;
+
+	/** The number of frames the subsystem waits to update added metrics after their initialization. */
+	UPROPERTY(Config)
+	int32 WarmUpFrames = 8;
+
+private:
+	int32 PendingWarmUpFrames = 0;
+
+	void InitializeMetrics();
+	void DeinitializeMetrics();
+
+	void OnUpdate(float DeltaTimeInSeconds);
+
+	/**
+	 * Returns the live world metric list index corresponding to the parameter metric class.
+	 *
+	 * @param InMetricClass: the class type of the world metric.
+	 * @return a valid index value or INDEX_NONE if not matching metric exists.
+	 */
+	int32 GetMetricIndex(const TSubclassOf<UWorldMetricInterface>& InMetricClass) const;
+
+	/**
+	 * Returns the live extension list index corresponding to the parameter extension class.
+	 *
+	 * @param InExtensionClass: the class type of the extension.
+	 * @return a valid index value or INDEX_NONE if not matching extension exists.
+	 */
+	int32 GetExtensionIndex(const TSubclassOf<UWorldMetricsExtension>& InExtensionClass) const;
+
+	UWorldMetricsExtension* AcquireExtensionInternal(
+		UObject* InOwner,
+		const TSubclassOf<UWorldMetricsExtension>& InExtensionClass);
+
+	UWorldMetricsExtension* AcquireExistingExtension(
+		UObject* InOwner,
+		const TSubclassOf<UWorldMetricsExtension>& InExtensionClass);
+
+	UWorldMetricsExtension* AddExtension(UObject* InOwner, const TSubclassOf<UWorldMetricsExtension>& InExtensionClass);
+
+	bool ReleaseExtensionInternal(UObject* InOwner, const TSubclassOf<UWorldMetricsExtension>& InExtensionClass);
+
+	bool TryRemoveExtensionAt(int32 ExtensionIndex);
+
+	/**
+	 * Verifies the parameter metric has released any acquired extensions, releasing any pending ones. It's the metric's
+	 * responsibility to release all previously acquired resources. This method logs all extensions owned by the
+	 * parameter metric that hasn't been released.
+	 *
+	 * @param InMetric: the metric to verify.
+	 */
+	void VerifyMetricReleasedAllExtensions(UWorldMetricInterface* InMetric);
+
+	/**
+	 * Verifies there are no orphan extensions. These are extensions that are still alive without any live metric. This
+	 * method logs and removes all orphan extensions, addressing the scenario where extensions that acquired other
+	 * extensions didn't release them. Extensions, like metrics, are responsible for releasing acquired extensions.
+	 */
+	void VerifyRemoveOrphanExtensions();
+};
diff --git a/Engine/Plugins/WorldMetrics/Source/WorldMetricsCore/WorldMetricsCore.Build.cs b/Engine/Plugins/WorldMetrics/Source/WorldMetricsCore/WorldMetricsCore.Build.cs
new file mode 100644
--- /dev/null
+++ b/Engine/Plugins/WorldMetrics/Source/WorldMetricsCore/WorldMetricsCore.Build.cs
@@ -0,0 +1,35 @@
+// Copyright Epic Games, Inc. All Rights Reserved.
+
+using System.IO;
+using UnrealBuildTool;
+
+public class WorldMetricsCore : ModuleRules
+{
+	public WorldMetricsCore(ReadOnlyTargetRules Target) : base(Target)
+	{
+		PCHUsage = PCHUsageMode.UseExplicitOrSharedPCHs;
+
+		/**
+		 * Add public dependencies that you statically link with here
+		 *	NOTE: Please use the "PrivateDependencyModuleNames" list to add your module unless
+		 *		it containing "types" that are required in this module's public headers.
+		 */
+
+		PublicDependencyModuleNames.AddRange
+		(
+			new string[]
+			{
+				"CoreUObject",
+			}
+		);
+
+		PrivateDependencyModuleNames.AddRange
+		(
+			new string[]
+			{
+				"Core",
+				"Engine",
+			}
+		);
+	}
+}
diff --git a/Engine/Plugins/WorldMetrics/Source/WorldMetricsTest/Private/WorldMetricsTest.cpp b/Engine/Plugins/WorldMetrics/Source/WorldMetricsTest/Private/WorldMetricsTest.cpp
new file mode 100644
--- /dev/null
+++ b/Engine/Plugins/WorldMetrics/Source/WorldMetricsTest/Private/WorldMetricsTest.cpp
@@ -0,0 +1,739 @@
+// Copyright Epic Games, Inc. All Rights Reserved.
+
+#if WITH_TESTS
+#include "Engine/Engine.h"
+#include "Engine/World.h"
+#include "Tests/TestHarnessAdapter.h"
+#endif	// WITH_TESTS
+
+#include "HAL/IConsoleManager.h"
+#include "WorldMetrics.h"
+#include "WorldMetricsLog.h"
+#include "WorldMetricsSubsystem.h"
+#include "WorldMetricsTestTypes.h"
+#include "WorldMetricsTestUtil.h"
+
+//---------------------------------------------------------------------------------------------------------------------
+// UE::WorldMetrics::Private
+//---------------------------------------------------------------------------------------------------------------------
+namespace UE::WorldMetrics::Private
+{
+
+void TestAll(UWorld* World);
+
+static FAutoConsoleCommandWithWorldAndArgs CmdWorldMetricsSelfTest(
+	TEXT("WorldMetrics.SelfTest"),
+	TEXT("Toggles the World Metrics Subsystem self-test."),
+	FConsoleCommandWithWorldAndArgsDelegate::CreateLambda([](const TArray<FString>& Args, UWorld* World)
+														  { TestAll(World); }),
+	ECVF_Default);
+
+static const TSubclassOf<UMockWorldMetricBase> MockMetricClasses[] = {
+	UMockWorldMetricA::StaticClass(), UMockWorldMetricB::StaticClass(), UMockWorldMetricC::StaticClass(),
+	UMockWorldMetricD::StaticClass(), UMockWorldMetricE::StaticClass(), UMockWorldMetricF::StaticClass(),
+	UMockWorldMetricG::StaticClass(), UMockWorldMetricH::StaticClass(), UMockWorldMetricI::StaticClass(),
+	UMockWorldMetricJ::StaticClass()};
+static constexpr int32 NumMockMetricClasses = sizeof(MockMetricClasses) / sizeof(MockMetricClasses[0]);
+
+static const TSubclassOf<UMockWorldMetricsExtensionBase> MockExtensionClasses[] = {
+	UMockWorldMetricsExtensionA::StaticClass(), UMockWorldMetricsExtensionB::StaticClass(),
+	UMockWorldMetricsExtensionC::StaticClass(), UMockWorldMetricsExtensionD::StaticClass(),
+	UMockWorldMetricsExtensionE::StaticClass(), UMockWorldMetricsExtensionF::StaticClass(),
+	UMockWorldMetricsExtensionG::StaticClass(), UMockWorldMetricsExtensionH::StaticClass(),
+	UMockWorldMetricsExtensionI::StaticClass(), UMockWorldMetricsExtensionJ::StaticClass()};
+
+static constexpr int32 NumMockExtensionClasses = sizeof(MockExtensionClasses) / sizeof(MockExtensionClasses[0]);
+
+TArray<TSubclassOf<UMockWorldMetricBase>> AddMockMetrics(UWorld* World, int32 Count, bool bRandomized = false)
+{
+	TArray<TSubclassOf<UMockWorldMetricBase>> Result;
+
+	UWorldMetricsSubsystem* Subsystem = UE::WorldMetrics::GetSubsystem(World);
+	if (UNLIKELY(!Subsystem))
+	{
+		return Result;
+	}
+	Count = FMath::Min(Count, NumMockMetricClasses);
+	Result.Reserve(Count);
+
+	const TArray<int32> ClassIndices =
+		bRandomized ? MakeRandomSubset<int32>(MakeIndexArray(NumMockMetricClasses), Count) : MakeIndexArray(Count);
+
+	for (int32 ClassIndex : ClassIndices)
+	{
+		if (Subsystem->AddMetric(MockMetricClasses[ClassIndex]))
+		{
+			Result.Emplace(MockMetricClasses[ClassIndex]);
+		}
+	}
+	return Result;
+}
+
+TArray<UMockWorldMetricBase*> GetMockMetrics(
+	UWorld* World,
+	TArrayView<TSubclassOf<UMockWorldMetricBase>> InMockMetricClasses)
+{
+	TArray<UMockWorldMetricBase*> Result;
+
+	UWorldMetricsSubsystem* Subsystem = UE::WorldMetrics::GetSubsystem(World);
+	if (UNLIKELY(!Subsystem))
+	{
+		return Result;
+	}
+
+	Result.Reserve(InMockMetricClasses.Num());
+	for (const TSubclassOf<UMockWorldMetricBase>& MetricClass : InMockMetricClasses)
+	{
+		if (UWorldMetricInterface* Metric = Subsystem->GetMetric(MetricClass))
+		{
+			Result.Emplace(static_cast<UMockWorldMetricBase*>(Metric));
+		}
+	}
+	return Result;
+}
+
+int32 RemoveMockMetrics(
+	UWorld* World,
+	TArrayView<TSubclassOf<UMockWorldMetricBase>> InMockMetricClasses,
+	bool bRandomized = false)
+{
+	UWorldMetricsSubsystem* Subsystem = UE::WorldMetrics::GetSubsystem(World);
+	if (UNLIKELY(!Subsystem))
+	{
+		return 0;
+	}
+
+	TArray<int32> ClassIndices = MakeIndexArray(InMockMetricClasses.Num());
+	if (bRandomized)
+	{
+		Shuffle(ClassIndices);
+	}
+
+	int32 Result = 0;
+	for (int32 ClassIndex : ClassIndices)
+	{
+		if (Subsystem->RemoveMetric(InMockMetricClasses[ClassIndex]))
+		{
+			++Result;
+		}
+	}
+	return Result;
+}
+
+TArray<TSubclassOf<UMockWorldMetricsExtensionBase>>
+AcquireMockMetricExtensions(UMockWorldMetricBase* Owner, int32 Count, bool bRandomized = false)
+{
+	TArray<TSubclassOf<UMockWorldMetricsExtensionBase>> Result;
+
+	if (UNLIKELY(!Owner))
+	{
+		return Result;
+	}
+
+	UWorldMetricsSubsystem* Subsystem = UE::WorldMetrics::GetSubsystem(Owner->GetWorld());
+	if (UNLIKELY(!Subsystem))
+	{
+		return Result;
+	}
+
+	Count = FMath::Min(Count, NumMockExtensionClasses);
+	Result.Reserve(Count);
+
+	const TArray<int32> ClassIndices =
+		bRandomized ? MakeRandomSubset<int32>(MakeIndexArray(NumMockExtensionClasses), Count) : MakeIndexArray(Count);
+
+	for (int32 ClassIndex : ClassIndices)
+	{
+		if (Subsystem->AcquireExtension(Owner, MockExtensionClasses[ClassIndex]))
+		{
+			Result.Emplace(MockExtensionClasses[ClassIndex]);
+		}
+	}
+	return Result;
+}
+
+int32 ReleaseMockMetricExtensions(
+	UMockWorldMetricBase* Owner,
+	TArrayView<TSubclassOf<UMockWorldMetricsExtensionBase>> InMockExtensionClasses,
+	bool bRandomized = false)
+{
+	if (UNLIKELY(!Owner))
+	{
+		return 0;
+	}
+
+	UWorldMetricsSubsystem* Subsystem = UE::WorldMetrics::GetSubsystem(Owner->GetWorld());
+	if (UNLIKELY(!Subsystem))
+	{
+		return 0;
+	}
+
+	TArray<int32> ClassIndices = MakeIndexArray(InMockExtensionClasses.Num());
+	if (bRandomized)
+	{
+		Shuffle(ClassIndices);
+	}
+
+	int32 Result = 0;
+	for (int32 ClassIndex : ClassIndices)
+	{
+		if (Subsystem->ReleaseExtension(Owner, InMockExtensionClasses[ClassIndex]))
+		{
+			++Result;
+		}
+	}
+	return Result;
+}
+
+void TestZeroState(UWorld* World)
+{
+	UWorldMetricsSubsystem* Subsystem = UE::WorldMetrics::GetSubsystem(World);
+	REQUIRE_MESSAGE(TEXT("WorldMetricsSubsystem is valid after GetSubsystem."), Subsystem);
+
+	REQUIRE_MESSAGE(TEXT("Unexpected existing metrics in WorldMetricsSubsystem"), !Subsystem->HasAnyMetric());
+	Subsystem->Enable(true);
+	REQUIRE_MESSAGE(
+		TEXT("WorldMetricsSubsystem shouldn't be enabled if there are no metrics."), !Subsystem->IsEnabled());
+	REQUIRE_MESSAGE(TEXT("WorldMetricsSubsystem should zero metrics."), Subsystem->NumMetrics() == 0);
+
+	REQUIRE_MESSAGE(TEXT("WorldMetricsSubsystem shouldn't be have any extension."), !Subsystem->HasAnyExtension());
+	REQUIRE_MESSAGE(TEXT("WorldMetricsSubsystem should have zero extensions."), Subsystem->NumExtensions() == 0);
+}
+
+void TestSingleMetricAddRemove(UWorld* World)
+{
+	TestZeroState(World);
+
+	// UE::WorldMetrics API
+	{
+		REQUIRE_MESSAGE(TEXT("AddMetric failed"), UE::WorldMetrics::AddMetric<UMockWorldMetricA>(World));
+
+		UMockWorldMetricA* MetricA = UE::WorldMetrics::GetMetric<UMockWorldMetricA>(World);
+		REQUIRE_MESSAGE(TEXT("Test Metric initialization failed"), MetricA->InitializeCount == 1);
+
+		const UWorldMetricsSubsystem* Subsystem = UE::WorldMetrics::GetSubsystem(World);
+		REQUIRE_MESSAGE(
+			TEXT("WorldMetricsSubsystem should be enabled by default if there are metrics."), Subsystem->IsEnabled());
+		REQUIRE_MESSAGE(TEXT("WorldMetricsSubsystem should have metrics."), Subsystem->HasAnyMetric());
+		REQUIRE_MESSAGE(TEXT("WorldMetricsSubsystem should have one metric."), Subsystem->NumMetrics() == 1);
+
+		REQUIRE_MESSAGE(TEXT("Unexpected AddMetric"), !UE::WorldMetrics::AddMetric<UMockWorldMetricA>(World));
+		REQUIRE_MESSAGE(TEXT("WorldMetricsSubsystem should have one metric."), Subsystem->NumMetrics() == 1);
+
+		REQUIRE_MESSAGE(TEXT("RemoveMetric failed."), UE::WorldMetrics::RemoveMetric<UMockWorldMetricA>(World));
+		REQUIRE_MESSAGE(TEXT("Test Metric deinitialization failed after release."), MetricA->DeinitializeCount == 1);
+
+		TestZeroState(World);
+
+		REQUIRE_MESSAGE(TEXT("Unexpected RemoveMetric."), !UE::WorldMetrics::RemoveMetric<UMockWorldMetricA>(World));
+
+		TestZeroState(World);
+	}
+
+	// World Metrics Subsystem API
+	{
+		UWorldMetricsSubsystem* Subsystem = UE::WorldMetrics::GetSubsystem(World);
+
+		REQUIRE_MESSAGE(TEXT("AddMetric failed"), Subsystem->AddMetric<UMockWorldMetricA>());
+
+		UMockWorldMetricA* MetricA = Subsystem->GetMetric<UMockWorldMetricA>();
+		REQUIRE_MESSAGE(TEXT("Test Metric initialization failed"), MetricA->InitializeCount == 1);
+
+		REQUIRE_MESSAGE(
+			TEXT("WorldMetricsSubsystem should be enabled by default if there are metrics."), Subsystem->IsEnabled());
+		REQUIRE_MESSAGE(TEXT("WorldMetricsSubsystem should have metrics."), Subsystem->HasAnyMetric());
+		REQUIRE_MESSAGE(TEXT("WorldMetricsSubsystem should have one metric."), Subsystem->NumMetrics() == 1);
+
+		REQUIRE_MESSAGE(TEXT("Unexpected AddMetric"), !Subsystem->AddMetric<UMockWorldMetricA>());
+		REQUIRE_MESSAGE(TEXT("WorldMetricsSubsystem should have one metric."), Subsystem->NumMetrics() == 1);
+
+		REQUIRE_MESSAGE(TEXT("RemoveMetric failed."), Subsystem->RemoveMetric<UMockWorldMetricA>());
+		REQUIRE_MESSAGE(TEXT("Test Metric deinitialization failed after release."), MetricA->DeinitializeCount == 1);
+
+		TestZeroState(World);
+
+		REQUIRE_MESSAGE(TEXT("Unexpected RemoveMetric."), !Subsystem->RemoveMetric<UMockWorldMetricA>());
+
+		TestZeroState(World);
+	}
+}
+
+void TestMultipleMetricsAddRemove(UWorld* World, bool bRandomized)
+{
+	TestZeroState(World);
+
+	constexpr int32 NumMetrics = 5;
+
+	auto MetricClasses = Private::AddMockMetrics(World, NumMetrics, bRandomized);
+
+	const UWorldMetricsSubsystem* Subsystem = UE::WorldMetrics::GetSubsystem(World);
+	REQUIRE_MESSAGE(
+		TEXT("WorldMetricsSubsystem should be enabled by default if there are metrics."), Subsystem->IsEnabled());
+	REQUIRE_MESSAGE(TEXT("WorldMetricsSubsystem should have extensions."), Subsystem->HasAnyMetric());
+	REQUIRE_MESSAGE(TEXT("WorldMetricsSubsystem should have one extension."), Subsystem->NumMetrics() == NumMetrics);
+
+	TArray<UMockWorldMetricBase*> MockMetrics = Private::GetMockMetrics(World, MetricClasses);
+	MockMetrics.Reserve(MetricClasses.Num());
+	for (const UMockWorldMetricBase* MockMetric : MockMetrics)
+	{
+		REQUIRE_MESSAGE(TEXT("Test Metric initialization failed"), MockMetric->InitializeCount == 1);
+	}
+
+	const int32 NumMetricsRemoved = Private::RemoveMockMetrics(World, MetricClasses, bRandomized);
+	REQUIRE_MESSAGE(TEXT("Unexpected metric remove count"), NumMetricsRemoved == NumMetrics);
+	for (const UMockWorldMetricBase* Metric : MockMetrics)
+	{
+		REQUIRE_MESSAGE(TEXT("Test Metric initialization failed"), Metric->DeinitializeCount == 1);
+	}
+
+	TestZeroState(World);
+}
+
+void TestSingleMetricSingleExtensionAcquireRelease(UWorld* World)
+{
+	TestZeroState(World);
+
+	REQUIRE_MESSAGE(TEXT("AddMetric failed"), UE::WorldMetrics::AddMetric<UMockWorldMetricA>(World));
+
+	UMockWorldMetricA* MetricA = UE::WorldMetrics::GetMetric<UMockWorldMetricA>(World);
+	REQUIRE_MESSAGE(TEXT("Test Metric initialization failed"), MetricA->InitializeCount == 1);
+
+	// UE::WorldMetrics API
+	{
+		UMockWorldMetricsExtensionA* ExtensionA =
+			UE::WorldMetrics::AcquireExtension<UMockWorldMetricsExtensionA>(MetricA);
+
+		REQUIRE_MESSAGE(TEXT("Acquire Test Extension failed."), ExtensionA);
+		REQUIRE_MESSAGE(TEXT("Test Extension initialization failed"), ExtensionA->InitializeCount == 1);
+
+		const UWorldMetricsSubsystem* Subsystem = UE::WorldMetrics::GetSubsystem(World);
+		REQUIRE_MESSAGE(
+			TEXT("WorldMetricsSubsystem should be enabled by default if there are metrics."), Subsystem->IsEnabled());
+		REQUIRE_MESSAGE(TEXT("WorldMetricsSubsystem should have extensions."), Subsystem->HasAnyExtension());
+		REQUIRE_MESSAGE(TEXT("WorldMetricsSubsystem should have one extension."), Subsystem->NumExtensions() == 1);
+
+		REQUIRE_MESSAGE(
+			TEXT("Test Extension release failed."),
+			UE::WorldMetrics::ReleaseExtension<UMockWorldMetricsExtensionA>(MetricA));
+		REQUIRE_MESSAGE(
+			TEXT("Test Extension deinitialization failed after release."), ExtensionA->DeinitializeCount == 1);
+		REQUIRE_MESSAGE(TEXT("WorldMetricsSubsystem shouldn't be have any extension."), !Subsystem->HasAnyExtension());
+		REQUIRE_MESSAGE(TEXT("WorldMetricsSubsystem should have zero extensions."), Subsystem->NumExtensions() == 0);
+	}
+
+	// World Metrics Subsystem API
+	{
+		UWorldMetricsSubsystem* Subsystem = UE::WorldMetrics::GetSubsystem(World);
+
+		UMockWorldMetricsExtensionA* ExtensionA = Subsystem->AcquireExtension<UMockWorldMetricsExtensionA>(MetricA);
+		REQUIRE_MESSAGE(TEXT("Acquire Test Extension failed."), ExtensionA);
+		REQUIRE_MESSAGE(TEXT("Test Extension initialization failed"), ExtensionA->InitializeCount == 1);
+
+		REQUIRE_MESSAGE(
+			TEXT("WorldMetricsSubsystem should be enabled by default if there are metrics."), Subsystem->IsEnabled());
+		REQUIRE_MESSAGE(TEXT("WorldMetricsSubsystem should have extensions."), Subsystem->HasAnyExtension());
+		REQUIRE_MESSAGE(TEXT("WorldMetricsSubsystem should have one extension."), Subsystem->NumExtensions() == 1);
+
+		REQUIRE_MESSAGE(
+			TEXT("Test Extension release failed."), Subsystem->ReleaseExtension<UMockWorldMetricsExtensionA>(MetricA));
+		REQUIRE_MESSAGE(
+			TEXT("Test Extension deinitialization failed after release."), ExtensionA->DeinitializeCount == 1);
+	}
+
+	UE::WorldMetrics::RemoveMetric<UMockWorldMetricA>(World);
+
+	TestZeroState(World);
+}
+
+void TestMultipleMetricsSingleExtensionAcquireRelease(UWorld* World, bool bRandomized)
+{
+	TestZeroState(World);
+
+	constexpr int32 NumMetrics = 5;
+
+	auto MetricClasses = Private::AddMockMetrics(World, NumMetrics, bRandomized);
+	TArray<UMockWorldMetricBase*> MockMetrics = Private::GetMockMetrics(World, MetricClasses);
+	REQUIRE_MESSAGE(TEXT("Unexpected empty metrics list."), !MockMetrics.IsEmpty());
+
+	UMockWorldMetricsExtensionA* ExtensionA =
+		UE::WorldMetrics::AcquireExtension<UMockWorldMetricsExtensionA>(MockMetrics[0]);
+
+	// Acquire extension loop
+	{
+		int32 CheckCount = 1;
+		for (UMockWorldMetricBase* MockMetric : MockMetrics)
+		{
+			UE::WorldMetrics::AcquireExtension<UMockWorldMetricsExtensionA>(MockMetric);
+			REQUIRE_MESSAGE(TEXT("Acquire Test Extension failed."), ExtensionA);
+			REQUIRE_MESSAGE(TEXT("Test Extension initialization failed"), ExtensionA->InitializeCount == 1);
+			REQUIRE_MESSAGE(TEXT("Test Extension initialization failed"), ExtensionA->OnAcquireCount == ++CheckCount);
+		}
+	}
+	// Release extension loop
+	{
+		if (bRandomized)
+		{
+			Shuffle(MockMetrics);
+		}
+
+		int32 CheckCount = 0;
+		for (UMockWorldMetricBase* MockMetric : MockMetrics)
+		{
+			REQUIRE_MESSAGE(TEXT("Test Extension unexpected deinitialization"), ExtensionA->DeinitializeCount == 0);
+			UE::WorldMetrics::ReleaseExtension<UMockWorldMetricsExtensionA>(MockMetric);
+
+			REQUIRE_MESSAGE(TEXT("Acquire Test Extension failed."), ExtensionA);
+			REQUIRE_MESSAGE(TEXT("Test Extension initialization failed"), ExtensionA->InitializeCount == 1);
+			REQUIRE_MESSAGE(TEXT("Test Extension initialization failed"), ExtensionA->OnReleaseCount == ++CheckCount);
+		}
+		REQUIRE_MESSAGE(TEXT("Test Extension initialization failed"), ExtensionA->DeinitializeCount == 1);
+	}
+
+	Private::RemoveMockMetrics(World, MetricClasses, bRandomized);
+
+	TestZeroState(World);
+}
+
+void TestMultipleMetricsMultipleExtensionAcquireRelease(UWorld* World, bool bRandomized)
+{
+	TestZeroState(World);
+
+	constexpr int32 NumMetrics = 5;
+	constexpr int32 NumExtension = 5;
+
+	auto MetricClasses = Private::AddMockMetrics(World, NumMetrics, bRandomized);
+	TArray<UMockWorldMetricBase*> MockMetrics = Private::GetMockMetrics(World, MetricClasses);
+	REQUIRE_MESSAGE(TEXT("Unexpected empty metrics list."), !MockMetrics.IsEmpty());
+
+	TArray<TArray<TSubclassOf<UMockWorldMetricsExtensionBase>>> Extensions;
+	Extensions.Reserve(NumMetrics);
+
+	// Acquire extensions
+	for (UMockWorldMetricBase* MockMetric : MockMetrics)
+	{
+		Extensions.Emplace(Private::AcquireMockMetricExtensions(MockMetric, NumExtension, bRandomized));
+	}
+
+	const UWorldMetricsSubsystem* Subsystem = UE::WorldMetrics::GetSubsystem(World);
+	REQUIRE_MESSAGE(
+		TEXT("WorldMetricsSubsystem should be enabled by default if there are metrics."), Subsystem->IsEnabled());
+	REQUIRE_MESSAGE(TEXT("WorldMetricsSubsystem should have extensions."), Subsystem->HasAnyExtension());
+
+	// Release extensions
+	for (int32 MetricIndex = 0; MetricIndex < NumMetrics; ++MetricIndex)
+	{
+		Private::ReleaseMockMetricExtensions(MockMetrics[MetricIndex], Extensions[MetricIndex], bRandomized);
+	}
+
+	REQUIRE_MESSAGE(TEXT("WorldMetricsSubsystem shouldn't be have any extension."), !Subsystem->HasAnyExtension());
+	REQUIRE_MESSAGE(TEXT("WorldMetricsSubsystem should have zero extensions."), Subsystem->NumExtensions() == 0);
+
+	Private::RemoveMockMetrics(World, MetricClasses, bRandomized);
+
+	TestZeroState(World);
+}
+
+void TestMultipleMetricsEnable(UWorld* World, bool bRandomized)
+{
+	TestZeroState(World);
+
+	constexpr int32 NumMetrics = 5;
+
+	auto MetricClasses = Private::AddMockMetrics(World, NumMetrics, bRandomized);
+
+	UWorldMetricsSubsystem* Subsystem = UE::WorldMetrics::GetSubsystem(World);
+	REQUIRE_MESSAGE(
+		TEXT("WorldMetricsSubsystem should be enabled by default if there are metrics."), Subsystem->IsEnabled());
+	REQUIRE_MESSAGE(TEXT("WorldMetricsSubsystem should have extensions."), Subsystem->HasAnyMetric());
+	REQUIRE_MESSAGE(TEXT("WorldMetricsSubsystem should have one extension."), Subsystem->NumMetrics() == NumMetrics);
+
+	TArray<UMockWorldMetricBase*> MockMetrics = Private::GetMockMetrics(World, MetricClasses);
+	MockMetrics.Reserve(MetricClasses.Num());
+	for (const UMockWorldMetricBase* MockMetric : MockMetrics)
+	{
+		REQUIRE_MESSAGE(TEXT("Test Metric initialization failed"), MockMetric->InitializeCount == 1);
+	}
+
+	Subsystem->Enable(false);
+	for (const UMockWorldMetricBase* MockMetric : MockMetrics)
+	{
+		REQUIRE_MESSAGE(TEXT("Test Metric initialization failed"), MockMetric->DeinitializeCount == 1);
+	}
+
+	Subsystem->Enable(true);
+	for (const UMockWorldMetricBase* MockMetric : MockMetrics)
+	{
+		REQUIRE_MESSAGE(TEXT("Test Metric initialization failed"), MockMetric->InitializeCount == 2);
+	}
+
+	const int32 NumMetricsRemoved = Private::RemoveMockMetrics(World, MetricClasses, bRandomized);
+	REQUIRE_MESSAGE(TEXT("Unexpected metric remove count"), NumMetricsRemoved == NumMetrics);
+	for (const UMockWorldMetricBase* Metric : MockMetrics)
+	{
+		REQUIRE_MESSAGE(TEXT("Test Metric initialization failed"), Metric->DeinitializeCount == 2);
+	}
+
+	TestZeroState(World);
+}
+
+void TestSingleMetricSingleExtensionAutoReleaseOnRemoval(UWorld* World)
+{
+	TestZeroState(World);
+
+	REQUIRE_MESSAGE(TEXT("AddMetric failed"), UE::WorldMetrics::AddMetric<UMockWorldMetricA>(World));
+
+	UMockWorldMetricA* MetricA = UE::WorldMetrics::GetMetric<UMockWorldMetricA>(World);
+
+	UMockWorldMetricsExtensionA* ExtensionA = UE::WorldMetrics::AcquireExtension<UMockWorldMetricsExtensionA>(MetricA);
+	REQUIRE_MESSAGE(TEXT("Acquire Test Extension failed."), ExtensionA);
+	REQUIRE_MESSAGE(TEXT("Test Extension initialization failed"), ExtensionA->InitializeCount == 1);
+
+	const UWorldMetricsSubsystem* Subsystem = UE::WorldMetrics::GetSubsystem(World);
+	REQUIRE_MESSAGE(
+		TEXT("WorldMetricsSubsystem should be enabled by default if there are metrics."), Subsystem->IsEnabled());
+	REQUIRE_MESSAGE(TEXT("WorldMetricsSubsystem should have extensions."), Subsystem->HasAnyExtension());
+	REQUIRE_MESSAGE(TEXT("WorldMetricsSubsystem should have one extension."), Subsystem->NumExtensions() == 1);
+
+	UE::WorldMetrics::RemoveMetric<UMockWorldMetricA>(World);
+
+	REQUIRE_MESSAGE(TEXT("Test Extension deinitialization failed after release."), ExtensionA->DeinitializeCount == 1);
+	REQUIRE_MESSAGE(TEXT("WorldMetricsSubsystem shouldn't have any extensions."), !Subsystem->HasAnyExtension());
+
+	TestZeroState(World);
+}
+
+void TestMultipleMetricsMultipleExtensionAutoReleaseOnRemoval(UWorld* World, bool bRandomized)
+{
+	TestZeroState(World);
+
+	constexpr int32 NumMetrics = 5;
+	constexpr int32 NumExtension = 5;
+
+	auto MetricClasses = Private::AddMockMetrics(World, NumMetrics, bRandomized);
+	TArray<UMockWorldMetricBase*> MockMetrics = Private::GetMockMetrics(World, MetricClasses);
+	REQUIRE_MESSAGE(TEXT("Unexpected empty metrics list."), !MockMetrics.IsEmpty());
+
+	// Acquire extensions
+	for (UMockWorldMetricBase* MockMetric : MockMetrics)
+	{
+		Private::AcquireMockMetricExtensions(MockMetric, NumExtension, bRandomized);
+	}
+
+	const UWorldMetricsSubsystem* Subsystem = UE::WorldMetrics::GetSubsystem(World);
+	REQUIRE_MESSAGE(
+		TEXT("WorldMetricsSubsystem should be enabled by default if there are metrics."), Subsystem->IsEnabled());
+	REQUIRE_MESSAGE(TEXT("WorldMetricsSubsystem should have extensions."), Subsystem->HasAnyExtension());
+
+	Private::RemoveMockMetrics(World, MetricClasses, bRandomized);
+
+	TestZeroState(World);
+}
+
+void TestSingleMetricOrphanExtensionAutoRemoval(UWorld* World)
+{
+	TestZeroState(World);
+
+	REQUIRE_MESSAGE(TEXT("AddMetric failed"), UE::WorldMetrics::AddMetric<UMockWorldMetricA>(World));
+
+	UMockWorldMetricA* MetricA = UE::WorldMetrics::GetMetric<UMockWorldMetricA>(World);
+
+	UMockWorldMetricsExtensionA* ExtensionA = UE::WorldMetrics::AcquireExtension<UMockWorldMetricsExtensionA>(MetricA);
+	REQUIRE_MESSAGE(TEXT("Acquire Test Extension failed."), ExtensionA);
+	REQUIRE_MESSAGE(TEXT("Test Extension initialization failed"), ExtensionA->InitializeCount == 1);
+
+	// Extension dependency: ExtensionB requires ExtensionA
+	UMockWorldMetricsExtensionB* ExtensionB = UE::WorldMetrics::AcquireExtension<UMockWorldMetricsExtensionB>(ExtensionA);
+	REQUIRE_MESSAGE(TEXT("Acquire Test Extension failed."), ExtensionB);
+	REQUIRE_MESSAGE(TEXT("Test Extension initialization failed"), ExtensionB->InitializeCount == 1);
+
+	const UWorldMetricsSubsystem* Subsystem = UE::WorldMetrics::GetSubsystem(World);
+	REQUIRE_MESSAGE(
+		TEXT("WorldMetricsSubsystem should be enabled by default if there are metrics."), Subsystem->IsEnabled());
+	REQUIRE_MESSAGE(TEXT("WorldMetricsSubsystem should have extensions."), Subsystem->HasAnyExtension());
+	REQUIRE_MESSAGE(TEXT("WorldMetricsSubsystem should have one extension."), Subsystem->NumExtensions() == 2);
+
+	UE::WorldMetrics::RemoveMetric<UMockWorldMetricA>(World);
+
+	REQUIRE_MESSAGE(TEXT("Test Extension deinitialization failed after release."), ExtensionA->DeinitializeCount == 1);
+	REQUIRE_MESSAGE(TEXT("Test Extension deinitialization failed after release."), ExtensionB->DeinitializeCount == 1);
+	REQUIRE_MESSAGE(TEXT("WorldMetricsSubsystem shouldn't have any extensions."), !Subsystem->HasAnyExtension());
+
+	TestZeroState(World);
+}
+
+void TestAll(UWorld* World)
+{
+	UWorldMetricsSubsystem* Subsystem = UE::WorldMetrics::GetSubsystem(World);
+	REQUIRE_MESSAGE(TEXT("WorldMetricsSubsystem is valid after GetSubsystem."), Subsystem);
+
+	Subsystem->Clear();
+
+	TestZeroState(World);
+
+	// Metrics
+	TestSingleMetricAddRemove(World);
+	TestMultipleMetricsAddRemove(World, false);
+	TestMultipleMetricsAddRemove(World, true);
+	TestMultipleMetricsEnable(World, false);
+	TestMultipleMetricsEnable(World, true);
+
+	// Extension
+	TestSingleMetricSingleExtensionAcquireRelease(World);
+	TestMultipleMetricsSingleExtensionAcquireRelease(World, false);
+	TestMultipleMetricsSingleExtensionAcquireRelease(World, true);
+	TestMultipleMetricsMultipleExtensionAcquireRelease(World, false);
+	TestMultipleMetricsMultipleExtensionAcquireRelease(World, true);
+
+	// Extension auto-removal when no longer acquired
+	TestSingleMetricSingleExtensionAutoReleaseOnRemoval(World);
+	TestMultipleMetricsMultipleExtensionAutoReleaseOnRemoval(World, false);
+	TestMultipleMetricsMultipleExtensionAutoReleaseOnRemoval(World, true);
+	TestSingleMetricOrphanExtensionAutoRemoval(World);
+}
+
+}  // namespace UE::WorldMetrics::Private
+
+//---------------------------------------------------------------------------------------------------------------------
+// WITH_TESTS
+//---------------------------------------------------------------------------------------------------------------------
+
+#if WITH_TESTS
+namespace UE::WorldMetrics
+{
+namespace Private
+{
+static void ScopedWorldTest(EWorldType::Type WorldType, TFunctionRef<void(UWorld* World)> WorldTest)
+{
+	constexpr bool bInformEngineOfWorld = false;
+	UWorld* World = UWorld::CreateWorld(WorldType, bInformEngineOfWorld);
+	REQUIRE_MESSAGE(TEXT("World could not be created."), World != nullptr);
+
+	FWorldContext& WorldContext = GEngine->CreateNewWorldContext(WorldType);
+	WorldContext.SetCurrentWorld(World);
+
+	const FURL URL;
+	World->InitializeActorsForPlay(URL);
+	World->BeginPlay();
+
+	WorldTest(World);
+
+	GEngine->DestroyWorldContext(World);
+	World->DestroyWorld(bInformEngineOfWorld);
+}
+}  // namespace Private
+
+TEST_CASE_NAMED(WorldMetricsTestZeroState, "WorldMetrics::TestZeroState", "[WorldMetrics][ClientContext][EngineFilter]")
+{
+	Private::ScopedWorldTest(EWorldType::Editor, [this](UWorld* World) { Private::TestZeroState(World); });
+}
+
+TEST_CASE_NAMED(
+	WorldMetricsTestSingleMetricAddRemove,
+	"WorldMetrics::TestSingleMetricAddRemove",
+	"[WorldMetrics][ClientContext][EngineFilter]")
+{
+	Private::ScopedWorldTest(EWorldType::Editor, [this](UWorld* World) { Private::TestSingleMetricAddRemove(World); });
+}
+
+TEST_CASE_NAMED(
+	WorldMetricsTestMultipleMetricsAddRemove,
+	"WorldMetrics::TestMultipleMetricsAddRemove",
+	"[WorldMetrics][ClientContext][EngineFilter]")
+{
+	Private::ScopedWorldTest(
+		EWorldType::Editor,
+		[this](UWorld* World)
+		{
+			Private::TestMultipleMetricsAddRemove(World, false);
+			Private::TestMultipleMetricsAddRemove(World, true);
+		});
+}
+
+TEST_CASE_NAMED(
+	WorldMetricsTestMultipleMetricsEnable,
+	"WorldMetrics::TestMultipleMetricsEnable",
+	"[WorldMetrics][ClientContext][EngineFilter]")
+{
+	Private::ScopedWorldTest(
+		EWorldType::Editor,
+		[this](UWorld* World)
+		{
+			Private::TestMultipleMetricsEnable(World, false);
+			Private::TestMultipleMetricsEnable(World, true);
+		});
+}
+
+TEST_CASE_NAMED(
+	WorldMetricsTestSingleMetricSingleExtensionAcquireRelease,
+	"WorldMetrics::TestSingleMetricSingleExtensionAcquireRelease",
+	"[WorldMetrics][ClientContext][EngineFilter]")
+{
+	Private::ScopedWorldTest(
+		EWorldType::Editor, [this](UWorld* World) { Private::TestSingleMetricSingleExtensionAcquireRelease(World); });
+}
+
+TEST_CASE_NAMED(
+	WorldMetricsTestMultipleMetricsSingleExtensionAcquireRelease,
+	"WorldMetrics::TestMultipleMetricsSingleExtensionAcquireRelease",
+	"[WorldMetrics][ClientContext][EngineFilter]")
+{
+	Private::ScopedWorldTest(
+		EWorldType::Editor,
+		[this](UWorld* World)
+		{
+			Private::TestMultipleMetricsSingleExtensionAcquireRelease(World, false);
+			Private::TestMultipleMetricsSingleExtensionAcquireRelease(World, true);
+		});
+}
+
+TEST_CASE_NAMED(
+	WorldMetricsTestMultipleMetricsMultipleExtensionAcquireRelease,
+	"WorldMetrics::TestMultipleMetricsMultipleExtensionAcquireRelease",
+	"[WorldMetrics][ClientContext][EngineFilter]")
+{
+	Private::ScopedWorldTest(
+		EWorldType::Editor,
+		[this](UWorld* World)
+		{
+			Private::TestMultipleMetricsMultipleExtensionAcquireRelease(World, false);
+			Private::TestMultipleMetricsMultipleExtensionAcquireRelease(World, true);
+		});
+}
+
+TEST_CASE_NAMED(
+	WorldMetricsTestSingleMetricSingleExtensionAutoReleaseOnRemoval,
+	"WorldMetrics::TestSingleMetricSingleExtensionAutoReleaseOnRemoval",
+	"[WorldMetrics][ClientContext][EngineFilter]")
+{
+	Private::ScopedWorldTest(
+		EWorldType::Editor,
+		[this](UWorld* World) { Private::TestSingleMetricSingleExtensionAutoReleaseOnRemoval(World); });
+}
+
+TEST_CASE_NAMED(
+	WorldMetricsTestMultipleMetricsMultipleExtensionAutoReleaseOnRemoval,
+	"WorldMetrics::TestMultipleMetricsMultipleExtensionAutoReleaseOnRemoval",
+	"[WorldMetrics][ClientContext][EngineFilter]")
+{
+	Private::ScopedWorldTest(
+		EWorldType::Editor,
+		[this](UWorld* World)
+		{
+			Private::TestMultipleMetricsMultipleExtensionAutoReleaseOnRemoval(World, false);
+			Private::TestMultipleMetricsMultipleExtensionAutoReleaseOnRemoval(World, true);
+		});
+}
+
+TEST_CASE_NAMED(
+	WorldMetricsTestSingleMetricOrphanExtensionAutoRemoval,
+	"WorldMetrics::TestSingleMetricOrphanExtensionAutoRemoval",
+	"[WorldMetrics][ClientContext][EngineFilter]")
+{
+	Private::ScopedWorldTest(
+		EWorldType::Editor,
+		[this](UWorld* World) { Private::TestSingleMetricOrphanExtensionAutoRemoval(World); });
+}
+
+}  // namespace UE::WorldMetrics
+
+#endif	// WITH_TESTS
diff --git a/Engine/Plugins/WorldMetrics/Source/WorldMetricsTest/Private/WorldMetricsTestModule.cpp b/Engine/Plugins/WorldMetrics/Source/WorldMetricsTest/Private/WorldMetricsTestModule.cpp
new file mode 100644
--- /dev/null
+++ b/Engine/Plugins/WorldMetrics/Source/WorldMetricsTest/Private/WorldMetricsTestModule.cpp
@@ -0,0 +1,5 @@
+// Copyright Epic Games, Inc. All Rights Reserved.
+
+#include "Modules/ModuleManager.h"
+
+IMPLEMENT_MODULE(FDefaultModuleImpl, WorldMetricsTest)
diff --git a/Engine/Plugins/WorldMetrics/Source/WorldMetricsTest/Private/WorldMetricsTestTypes.cpp b/Engine/Plugins/WorldMetrics/Source/WorldMetricsTest/Private/WorldMetricsTestTypes.cpp
new file mode 100644
--- /dev/null
+++ b/Engine/Plugins/WorldMetrics/Source/WorldMetricsTest/Private/WorldMetricsTestTypes.cpp
@@ -0,0 +1,48 @@
+// Copyright Epic Games, Inc. All Rights Reserved.
+
+#include "WorldMetricsTestTypes.h"
+
+#include UE_INLINE_GENERATED_CPP_BY_NAME(WorldMetricsTestTypes)
+
+//---------------------------------------------------------------------------------------------------------------------
+// UMockWorldMetricBase
+//---------------------------------------------------------------------------------------------------------------------
+
+void UMockWorldMetricBase::Initialize()
+{
+	++InitializeCount;
+}
+
+void UMockWorldMetricBase::Deinitialize()
+{
+	++DeinitializeCount;
+}
+
+void UMockWorldMetricBase::Update(float /*DeltaTimeInSeconds*/)
+{
+	++UpdateCount;
+}
+
+//---------------------------------------------------------------------------------------------------------------------
+// UMockWorldMetricsExtensionBase
+//---------------------------------------------------------------------------------------------------------------------
+
+void UMockWorldMetricsExtensionBase::Initialize()
+{
+	++InitializeCount;
+}
+
+void UMockWorldMetricsExtensionBase::Deinitialize()
+{
+	++DeinitializeCount;
+}
+
+void UMockWorldMetricsExtensionBase::OnAcquire(UObject* InOwner)
+{
+	++OnAcquireCount;
+}
+
+void UMockWorldMetricsExtensionBase::OnRelease(UObject* InOwner)
+{
+	++OnReleaseCount;
+}
diff --git a/Engine/Plugins/WorldMetrics/Source/WorldMetricsTest/Public/WorldMetricsTestTypes.h b/Engine/Plugins/WorldMetrics/Source/WorldMetricsTest/Public/WorldMetricsTestTypes.h
new file mode 100644
--- /dev/null
+++ b/Engine/Plugins/WorldMetrics/Source/WorldMetricsTest/Public/WorldMetricsTestTypes.h
@@ -0,0 +1,171 @@
+// Copyright Epic Games, Inc. All Rights Reserved.
+
+#pragma once
+
+#include "WorldMetricInterface.h"
+#include "WorldMetricsExtension.h"
+
+#include "WorldMetricsTestTypes.generated.h"
+
+UCLASS(abstract)
+class UMockWorldMetricBase : public UWorldMetricInterface
+{
+	GENERATED_BODY()
+
+	virtual void Initialize() override;
+	virtual void Deinitialize() override;
+	virtual void Update(float DeltaTimeInSeconds) override;
+
+public:
+	TArray<UWorldMetricsExtension*> Extensions;
+	int32 InitializeCount = 0;
+	int32 DeinitializeCount = 0;
+	int32 UpdateCount = 0;
+
+	virtual SIZE_T GetAllocatedSize() const override
+	{
+		return 0;
+	}
+};
+
+UCLASS(abstract)
+class UMockWorldMetricsExtensionBase : public UWorldMetricsExtension
+{
+	GENERATED_BODY()
+
+	virtual void Initialize() override;
+	virtual void Deinitialize() override;
+	virtual void OnAcquire(UObject* InOwner) override;
+	virtual void OnRelease(UObject* InOwner) override;
+
+public:
+	int32 InitializeCount = 0;
+	int32 DeinitializeCount = 0;
+	int32 OnAcquireCount = 0;
+	int32 OnReleaseCount = 0;
+
+	virtual SIZE_T GetAllocatedSize() const override
+	{
+		return 0;
+	}
+};
+
+UCLASS()
+class UMockWorldMetricA : public UMockWorldMetricBase
+{
+	GENERATED_BODY()
+};
+
+UCLASS()
+class UMockWorldMetricB : public UMockWorldMetricBase
+{
+	GENERATED_BODY()
+};
+
+UCLASS()
+class UMockWorldMetricC : public UMockWorldMetricBase
+{
+	GENERATED_BODY()
+};
+
+UCLASS()
+class UMockWorldMetricD : public UMockWorldMetricBase
+{
+	GENERATED_BODY()
+};
+
+UCLASS()
+class UMockWorldMetricE : public UMockWorldMetricBase
+{
+	GENERATED_BODY()
+};
+
+UCLASS()
+class UMockWorldMetricF : public UMockWorldMetricBase
+{
+	GENERATED_BODY()
+};
+
+UCLASS()
+class UMockWorldMetricG : public UMockWorldMetricBase
+{
+	GENERATED_BODY()
+};
+
+UCLASS()
+class UMockWorldMetricH : public UMockWorldMetricBase
+{
+	GENERATED_BODY()
+};
+
+UCLASS()
+class UMockWorldMetricI : public UMockWorldMetricBase
+{
+	GENERATED_BODY()
+};
+
+UCLASS()
+class UMockWorldMetricJ : public UMockWorldMetricBase
+{
+	GENERATED_BODY()
+};
+
+UCLASS()
+class UMockWorldMetricsExtensionA : public UMockWorldMetricsExtensionBase
+{
+	GENERATED_BODY()
+};
+
+UCLASS()
+class UMockWorldMetricsExtensionB : public UMockWorldMetricsExtensionBase
+{
+	GENERATED_BODY()
+};
+
+UCLASS()
+class UMockWorldMetricsExtensionC : public UMockWorldMetricsExtensionBase
+{
+	GENERATED_BODY()
+};
+
+UCLASS()
+class UMockWorldMetricsExtensionD : public UMockWorldMetricsExtensionBase
+{
+	GENERATED_BODY()
+};
+
+UCLASS()
+class UMockWorldMetricsExtensionE : public UMockWorldMetricsExtensionBase
+{
+	GENERATED_BODY()
+};
+
+UCLASS()
+class UMockWorldMetricsExtensionF : public UMockWorldMetricsExtensionBase
+{
+	GENERATED_BODY()
+};
+
+UCLASS()
+class UMockWorldMetricsExtensionG : public UMockWorldMetricsExtensionBase
+{
+	GENERATED_BODY()
+};
+
+UCLASS()
+class UMockWorldMetricsExtensionH : public UMockWorldMetricsExtensionBase
+{
+	GENERATED_BODY()
+};
+
+UCLASS()
+class UMockWorldMetricsExtensionI : public UMockWorldMetricsExtensionBase
+{
+	GENERATED_BODY()
+};
+
+UCLASS()
+class UMockWorldMetricsExtensionJ : public UMockWorldMetricsExtensionBase
+{
+	GENERATED_BODY()
+};
diff --git a/Engine/Plugins/WorldMetrics/Source/WorldMetricsTest/Public/WorldMetricsTestUtil.h b/Engine/Plugins/WorldMetrics/Source/WorldMetricsTest/Public/WorldMetricsTestUtil.h
new file mode 100644
--- /dev/null
+++ b/Engine/Plugins/WorldMetrics/Source/WorldMetricsTest/Public/WorldMetricsTestUtil.h
@@ -0,0 +1,83 @@
+// Copyright Epic Games, Inc. All Rights Reserved.
+
+#pragma once
+
+#ifndef REQUIRE_MESSAGE
+#define REQUIRE_MESSAGE(Message, Expr)           \
+	if (!(Expr))                                 \
+	{                                            \
+		UE_LOG(LogWorldMetrics, Error, Message); \
+		return;                                  \
+	}
+#endif
+
+namespace UE::WorldMetrics
+{
+
+/**
+ * Returns an array with the parameter number of sequentially increased index values starting from zero.
+ * This method is equivalent to using std::iota with a zero starting value in a sequential container.
+ */
+inline TArray<int32> MakeIndexArray(int32 Size)
+{
+	// TODO: consider using std::iota when <algorithm> is allowed
+	TArray<int32> Indices;
+	Indices.SetNum(Size);
+	for (int32 i = 0; i < Size; ++i)
+	{
+		Indices[i] = i;
+	}
+	return Indices;
+}
+
+/**
+ * Returns an array with a subset of the elements in the parameter input array. The subset size is expressed
+ * by the parameter size which must be smaller than the number of elements in the input array.
+ */
+template <typename T>
+TArray<T> MakeRandomSubset(TArrayView<T> Items, int32 Size = 0)
+{
+	if (!Size)
+	{
+		Size = FMath::Max(1, Items.Num() / 2);
+	}
+	if (Items.Num() == Size)
+	{
+		return TArray<T>(Items);
+	}
+	TArray<T> Result;
+	if (Items.Num() > 1 && Size < Items.Num())
+	{
+		Result.Reset(Size);
+		while (Result.Num() < Size)
+		{
+			const int32 i = FMath::RandRange(0, Items.Num() - 1);
+			if (Result.Find(Items[i]) == INDEX_NONE)
+			{
+				Result.Add(Items[i]);
+			}
+		}
+	}
+	return Result;
+}
+
+/**
+ * Performs a number of randomized swaps in the parameter array equivalent to its number of elements.
+ */
+template <typename T>
+void Shuffle(TArray<T>& OutArray)
+{
+	if (OutArray.Num() > 1)
+	{
+		for (int32 i = 0; i < OutArray.Num(); i++)
+		{
+			const int32 j = FMath::RandRange(i, OutArray.Num() - 1);
+			if (i != j)
+			{
+				OutArray.Swap(i, j);
+			}
+		}
+	}
+}
+
+}  // namespace UE::WorldMetrics
diff --git a/Engine/Plugins/WorldMetrics/Source/WorldMetricsTest/WorldMetricsTest.Build.cs b/Engine/Plugins/WorldMetrics/Source/WorldMetricsTest/WorldMetricsTest.Build.cs
new file mode 100644
--- /dev/null
+++ b/Engine/Plugins/WorldMetrics/Source/WorldMetricsTest/WorldMetricsTest.Build.cs
@@ -0,0 +1,36 @@
+// Copyright Epic Games, Inc. All Rights Reserved.
+
+using System.IO;
+using UnrealBuildTool;
+
+public class WorldMetricsTest : ModuleRules
+{
+	public WorldMetricsTest(ReadOnlyTargetRules Target) : base(Target)
+	{
+		PCHUsage = PCHUsageMode.UseExplicitOrSharedPCHs;
+
+		/**
+		 * Add public dependencies that you statically link with here
+		 *	NOTE: Please use the "PrivateDependencyModuleNames" list to add your module unless
+		 *		it containing "types" that are required in this module's public headers.
+		 */
+
+		PublicDependencyModuleNames.AddRange
+		(
+			new string[]
+			{
+				"CoreUObject",
+				"WorldMetricsCore",
+			}
+		);
+
+		PrivateDependencyModuleNames.AddRange
+		(
+			new string[]
+			{
+				"Core",
+				"Engine",
+			}
+		);
+	}
+}
diff --git a/Engine/Plugins/WorldMetrics/WorldMetrics.uplugin b/Engine/Plugins/WorldMetrics/WorldMetrics.uplugin
new file mode 100644
--- /dev/null
+++ b/Engine/Plugins/WorldMetrics/WorldMetrics.uplugin
@@ -0,0 +1,37 @@
+{
+	"FileVersion": 3,
+	"Version": 1,
+	"VersionName": "1.0",
+	"FriendlyName": "World Metrics",
+	"Description": "Provides access to the world metrics system",
+	"Category": "Profiling",
+	"CreatedBy": "Epic Games, Inc.",
+	"CreatedByURL": "http://epicgames.com",
+	"DocsURL": "",
+	"MarketplaceURL": "",
+	"SupportURL": "",
+	"CanContainContent": false,
+	"IsBetaVersion": false,
+	"IsExperimentalVersion": true,
+	"EnabledByDefault": true,
+	"Installed": false,
+	"EditorCustomVirtualPath": "WorldMetrics",
+	"Modules": [
+		{
+			"Name": "WorldMetricsCore",
+			"Type": "Runtime",
+			"LoadingPhase": "Default"
+		},
+		{
+			"Name": "WorldMetricsTest",
+			"Type": "Runtime",
+			"LoadingPhase": "Default",
+			"TargetConfigurationDenyList": [ "Shipping" ]
+		},
+		{
+			"Name": "CsvMetrics",
+			"Type": "Runtime",
+			"LoadingPhase": "Default"
+		}
+	]
+}
