WorldMetrics API update: removed static API methods in favor of an owner accessor within each metric and extension. This simplifies code maintenance while better expressing the World Metrics Subsystem's ownership contract towards metrics and extensions.

#tests unit tests passed

diff --git a/Engine/Plugins/WorldMetrics/Source/CsvMetrics/Private/CsvActorCountMetric.cpp b/Engine/Plugins/WorldMetrics/Source/CsvMetrics/Private/CsvActorCountMetric.cpp
--- a/Engine/Plugins/WorldMetrics/Source/CsvMetrics/Private/CsvActorCountMetric.cpp
+++ b/Engine/Plugins/WorldMetrics/Source/CsvMetrics/Private/CsvActorCountMetric.cpp
@@ -5,7 +5,6 @@
 #include "GameFramework/Actor.h"
 #include "HAL/IConsoleManager.h"
 #include "ProfilingDebugging/CsvProfiler.h"
-#include "WorldMetrics.h"
 #include "WorldMetricsActorTracker.h"
 
 #include UE_INLINE_GENERATED_CPP_BY_NAME(CsvActorCountMetric)
@@ -27,12 +26,12 @@ static TAutoConsoleVariable<int32> CVarCsvRecordActorCountThreshold(
 
 void UCsvActorCountMetric::Initialize()
 {
-	UE::WorldMetrics::AcquireExtension<UWorldMetricsActorTracker>(this);
+	GetOwner().AcquireExtension<UWorldMetricsActorTracker>(this);
 }
 
 void UCsvActorCountMetric::Deinitialize()
 {
-	UE::WorldMetrics::ReleaseExtension<UWorldMetricsActorTracker>(this);
+	GetOwner().ReleaseExtension<UWorldMetricsActorTracker>(this);
 
 	ActorClassNameCounter.Reset();
 	TotalActorCount = 0;
diff --git a/Engine/Plugins/WorldMetrics/Source/CsvMetrics/Private/CsvMetricsSubsystem.cpp b/Engine/Plugins/WorldMetrics/Source/CsvMetrics/Private/CsvMetricsSubsystem.cpp
--- a/Engine/Plugins/WorldMetrics/Source/CsvMetrics/Private/CsvMetricsSubsystem.cpp
+++ b/Engine/Plugins/WorldMetrics/Source/CsvMetrics/Private/CsvMetricsSubsystem.cpp
@@ -4,7 +4,7 @@
 
 #include "Engine/World.h"
 #include "ProfilingDebugging/CsvProfiler.h"
-#include "WorldMetrics.h"
+#include "WorldMetricsSubsystem.h"
 
 #include UE_INLINE_GENERATED_CPP_BY_NAME(CsvMetricsSubsystem)
 
@@ -72,22 +72,24 @@ void UCsvMetricsSubsystem::UnbindProfilerCallbacks()
 
 void UCsvMetricsSubsystem::AddMetrics()
 {
-	UWorld* World = GetWorld();
-	check(World);
-
-	for (const TSubclassOf<UWorldMetricInterface>& MetricClass : Metrics)
+	UWorldMetricsSubsystem* WorldMetricsSubsystem = UWorldMetricsSubsystem::Get(GetWorld());
+	if (ensure(WorldMetricsSubsystem))
 	{
-		UE::WorldMetrics::AddMetric(World, MetricClass);
+		for (const TSubclassOf<UWorldMetricInterface>& MetricClass : Metrics)
+		{
+			WorldMetricsSubsystem->AddMetric(MetricClass);
+		}
 	}
 }
 
 void UCsvMetricsSubsystem::RemoveMetrics()
 {
-	UWorld* World = GetWorld();
-	check(World);
-
-	for (const TSubclassOf<UWorldMetricInterface>& MetricClass : Metrics)
+	UWorldMetricsSubsystem* WorldMetricsSubsystem = UWorldMetricsSubsystem::Get(GetWorld());
+	if (ensure(WorldMetricsSubsystem))
 	{
-		UE::WorldMetrics::RemoveMetric(World, MetricClass);
+		for (const TSubclassOf<UWorldMetricInterface>& MetricClass : Metrics)
+		{
+			WorldMetricsSubsystem->RemoveMetric(MetricClass);
+		}
 	}
 }
diff --git a/Engine/Plugins/WorldMetrics/Source/WorldMetricsCore/Private/WorldMetrics.cpp b/Engine/Plugins/WorldMetrics/Source/WorldMetricsCore/Private/WorldMetrics.cpp
deleted file mode 100644
diff --git a/Engine/Plugins/WorldMetrics/Source/WorldMetricsCore/Private/WorldMetricsConsoleCommands.cpp b/Engine/Plugins/WorldMetrics/Source/WorldMetricsCore/Private/WorldMetricsConsoleCommands.cpp
--- a/Engine/Plugins/WorldMetrics/Source/WorldMetricsCore/Private/WorldMetricsConsoleCommands.cpp
+++ b/Engine/Plugins/WorldMetrics/Source/WorldMetricsCore/Private/WorldMetricsConsoleCommands.cpp
@@ -7,7 +7,6 @@
 #include "HAL/IConsoleManager.h"
 #include "String/LexFromString.h"
 #include "WorldMetricInterface.h"
-#include "WorldMetrics.h"
 #include "WorldMetricsLog.h"
 #include "WorldMetricsSubsystem.h"
 
@@ -20,7 +19,7 @@ static FAutoConsoleCommandWithWorldAndArgs CmdWorldMetricsEnable(
 	FConsoleCommandWithWorldAndArgsDelegate::CreateLambda(
 		[](const TArray<FString>& Args, UWorld* World)
 		{
-			if (UWorldMetricsSubsystem* WorldMetrics = GetSubsystem(World))
+			if (UWorldMetricsSubsystem* WorldMetrics = UWorldMetricsSubsystem::Get(World))
 			{
 				bool bEnable = !WorldMetrics->IsEnabled();
 				if (Args.Num())
@@ -38,7 +37,7 @@ static FAutoConsoleCommandWithWorldAndArgs CmdWorldMetricsClear(
 	FConsoleCommandWithWorldAndArgsDelegate::CreateLambda(
 		[](const TArray<FString>& Args, UWorld* World)
 		{
-			if (UWorldMetricsSubsystem* WorldMetrics = GetSubsystem(World))
+			if (UWorldMetricsSubsystem* WorldMetrics = UWorldMetricsSubsystem::Get(World))
 			{
 				WorldMetrics->Clear();
 			}
@@ -51,7 +50,7 @@ static FAutoConsoleCommandWithWorldAndArgs CmdWorldMetricsSetUpdateRateInSeconds
 	FConsoleCommandWithWorldAndArgsDelegate::CreateLambda(
 		[](const TArray<FString>& Args, UWorld* World)
 		{
-			if (UWorldMetricsSubsystem* WorldMetrics = GetSubsystem(World))
+			if (UWorldMetricsSubsystem* WorldMetrics = UWorldMetricsSubsystem::Get(World))
 			{
 				float UpdateRateInSeconds = 0;
 				if (Args.Num() && Args[0].IsNumeric())
@@ -76,7 +75,13 @@ static void WorldMetricToggleCmdHandler(
 		return;
 	}
 
-	const bool bMetricExists = GetMetric(World, MetricClassPtr) != nullptr;
+	UWorldMetricsSubsystem* WorldMetricsSubsystem = UWorldMetricsSubsystem::Get(World);
+	if (!WorldMetricsSubsystem)
+	{
+		return;
+	}
+
+	const bool bMetricExists = WorldMetricsSubsystem->GetMetric(MetricClassPtr) != nullptr;
 	const FName ClassName = MetricClassPtr->GetFName();
 
 	bool bEnable = !bMetricExists;
@@ -92,13 +97,13 @@ static void WorldMetricToggleCmdHandler(
 		return;
 	}
 
-	if (bEnable && AddMetric(World, MetricClassPtr))
+	if (bEnable && WorldMetricsSubsystem->AddMetric(MetricClassPtr))
 	{
 		Ar.Log(*WriteToString<128>(ClassName, ANSITEXTVIEW(" added.")));
 		return;
 	}
 
-	if (!bEnable && RemoveMetric(World, MetricClassPtr))
+	if (!bEnable && WorldMetricsSubsystem->RemoveMetric(MetricClassPtr))
 	{
 		Ar.Log(*WriteToString<128>(ClassName, ANSITEXTVIEW(" removed.")));
 		return;
diff --git a/Engine/Plugins/WorldMetrics/Source/WorldMetricsCore/Private/WorldMetricsSubsystem.cpp b/Engine/Plugins/WorldMetrics/Source/WorldMetricsCore/Private/WorldMetricsSubsystem.cpp
--- a/Engine/Plugins/WorldMetrics/Source/WorldMetricsCore/Private/WorldMetricsSubsystem.cpp
+++ b/Engine/Plugins/WorldMetrics/Source/WorldMetricsCore/Private/WorldMetricsSubsystem.cpp
@@ -26,6 +26,15 @@ bool CanHaveWorldMetrics(const UWorld* World)
 // UWorldMetricsSubsystem
 //---------------------------------------------------------------------------------------------------------------------
 
+UWorldMetricsSubsystem* UWorldMetricsSubsystem::Get(const UWorld* World)
+{
+	if (LIKELY(World))
+	{
+		return World->GetSubsystem<UWorldMetricsSubsystem>();
+	}
+	return nullptr;
+}
+
 bool UWorldMetricsSubsystem::ShouldCreateSubsystem(UObject* Outer) const
 {
 	return UE::WorldMetrics::Private::CanHaveWorldMetrics(Cast<UWorld>(Outer));
diff --git a/Engine/Plugins/WorldMetrics/Source/WorldMetricsCore/Public/WorldMetricInterface.h b/Engine/Plugins/WorldMetrics/Source/WorldMetricsCore/Public/WorldMetricInterface.h
--- a/Engine/Plugins/WorldMetrics/Source/WorldMetricsCore/Public/WorldMetricInterface.h
+++ b/Engine/Plugins/WorldMetrics/Source/WorldMetricsCore/Public/WorldMetricInterface.h
@@ -11,7 +11,7 @@
  *
  * This is the required interface class to implement world metrics.
  */
-UCLASS(abstract, MinimalAPI)
+UCLASS(abstract, MinimalAPI, Within=WorldMetricsSubsystem)
 class UWorldMetricInterface : public UObject
 {
 	GENERATED_BODY()
@@ -24,6 +24,15 @@ public:
 	[[nodiscard]] WORLDMETRICSCORE_API virtual SIZE_T GetAllocatedSize() const
 		PURE_VIRTUAL(GetAllocatedSize, return 0;);
 
+	/*
+	 * Returns the owning World Metrics Subsystem which is expected to be valid during this object's lifetime.
+	 * @return a reference to the owning World Metrics Subsystem.
+	 */
+	[[nodiscard]] UWorldMetricsSubsystem& GetOwner() const
+	{
+		return *GetOuterUWorldMetricsSubsystem();
+	}
+
 private:
 	/**
 	 * Metric initialization method called by the World Metrics subsystem on metric addition.
diff --git a/Engine/Plugins/WorldMetrics/Source/WorldMetricsCore/Public/WorldMetrics.h b/Engine/Plugins/WorldMetrics/Source/WorldMetricsCore/Public/WorldMetrics.h
deleted file mode 100644
diff --git a/Engine/Plugins/WorldMetrics/Source/WorldMetricsCore/Public/WorldMetricsExtension.h b/Engine/Plugins/WorldMetrics/Source/WorldMetricsCore/Public/WorldMetricsExtension.h
--- a/Engine/Plugins/WorldMetrics/Source/WorldMetricsCore/Public/WorldMetricsExtension.h
+++ b/Engine/Plugins/WorldMetrics/Source/WorldMetricsCore/Public/WorldMetricsExtension.h
@@ -16,7 +16,7 @@
  * 3. Implement acquired/release semantics. The World Metrics subsystem automatically initializes extensions on
  * acquisition and deinitializes them on release. The subsystem may deallocate an extension that have no acquisitions.
  */
-UCLASS(abstract, MinimalAPI)
+UCLASS(abstract, MinimalAPI, Within=WorldMetricsSubsystem)
 class UWorldMetricsExtension : public UObject
 {
 	GENERATED_BODY()
@@ -29,6 +29,15 @@ public:
 	[[nodiscard]] WORLDMETRICSCORE_API virtual SIZE_T GetAllocatedSize() const
 		PURE_VIRTUAL(GetAllocatedSize, return 0;);
 
+	/*
+	 * Returns the owning World Metrics Subsystem which is expected to be valid during this object's lifetime.
+	 * @return a reference to the owning World Metrics Subsystem.
+	 */
+	[[nodiscard]] UWorldMetricsSubsystem& GetOwner() const
+	{
+		return *GetOuterUWorldMetricsSubsystem();
+	}
+
 private:
 	/**
 	 * Called by the World Metrics subsystem whenever a new extension is added.
diff --git a/Engine/Plugins/WorldMetrics/Source/WorldMetricsCore/Public/WorldMetricsSubsystem.h b/Engine/Plugins/WorldMetrics/Source/WorldMetricsCore/Public/WorldMetricsSubsystem.h
--- a/Engine/Plugins/WorldMetrics/Source/WorldMetricsCore/Public/WorldMetricsSubsystem.h
+++ b/Engine/Plugins/WorldMetrics/Source/WorldMetricsCore/Public/WorldMetricsSubsystem.h
@@ -25,6 +25,8 @@ class UWorldMetricsSubsystem final : public UWorldSubsystem
 	GENERATED_BODY()
 
 public:
+	[[nodiscard]] WORLDMETRICSCORE_API static UWorldMetricsSubsystem* Get(const UWorld* World);
+
 	//~ Begin USubsystem
 	WORLDMETRICSCORE_API virtual bool ShouldCreateSubsystem(UObject* Outer) const override;
 	WORLDMETRICSCORE_API virtual void Initialize(FSubsystemCollectionBase& Collection) override;
diff --git a/Engine/Plugins/WorldMetrics/Source/WorldMetricsTest/Private/WorldMetricsTest.cpp b/Engine/Plugins/WorldMetrics/Source/WorldMetricsTest/Private/WorldMetricsTest.cpp
--- a/Engine/Plugins/WorldMetrics/Source/WorldMetricsTest/Private/WorldMetricsTest.cpp
+++ b/Engine/Plugins/WorldMetrics/Source/WorldMetricsTest/Private/WorldMetricsTest.cpp
@@ -7,7 +7,6 @@
 #endif	// WITH_TESTS
 
 #include "HAL/IConsoleManager.h"
-#include "WorldMetrics.h"
 #include "WorldMetricsLog.h"
 #include "WorldMetricsSubsystem.h"
 #include "WorldMetricsTestTypes.h"
@@ -48,7 +47,7 @@ TArray<TSubclassOf<UMockWorldMetricBase>> AddMockMetrics(UWorld* World, int32 Co
 {
 	TArray<TSubclassOf<UMockWorldMetricBase>> Result;
 
-	UWorldMetricsSubsystem* Subsystem = UE::WorldMetrics::GetSubsystem(World);
+	UWorldMetricsSubsystem* Subsystem = UWorldMetricsSubsystem::Get(World);
 	if (UNLIKELY(!Subsystem))
 	{
 		return Result;
@@ -75,7 +74,7 @@ TArray<UMockWorldMetricBase*> GetMockMetrics(
 {
 	TArray<UMockWorldMetricBase*> Result;
 
-	UWorldMetricsSubsystem* Subsystem = UE::WorldMetrics::GetSubsystem(World);
+	UWorldMetricsSubsystem* Subsystem = UWorldMetricsSubsystem::Get(World);
 	if (UNLIKELY(!Subsystem))
 	{
 		return Result;
@@ -97,7 +96,7 @@ int32 RemoveMockMetrics(
 	TArrayView<TSubclassOf<UMockWorldMetricBase>> InMockMetricClasses,
 	bool bRandomized = false)
 {
-	UWorldMetricsSubsystem* Subsystem = UE::WorldMetrics::GetSubsystem(World);
+	UWorldMetricsSubsystem* Subsystem = UWorldMetricsSubsystem::Get(World);
 	if (UNLIKELY(!Subsystem))
 	{
 		return 0;
@@ -130,7 +129,7 @@ AcquireMockMetricExtensions(UMockWorldMetricBase* Owner, int32 Count, bool bRand
 		return Result;
 	}
 
-	UWorldMetricsSubsystem* Subsystem = UE::WorldMetrics::GetSubsystem(Owner->GetWorld());
+	UWorldMetricsSubsystem* Subsystem = UWorldMetricsSubsystem::Get(Owner->GetWorld());
 	if (UNLIKELY(!Subsystem))
 	{
 		return Result;
@@ -162,7 +161,7 @@ int32 ReleaseMockMetricExtensions(
 		return 0;
 	}
 
-	UWorldMetricsSubsystem* Subsystem = UE::WorldMetrics::GetSubsystem(Owner->GetWorld());
+	UWorldMetricsSubsystem* Subsystem = UWorldMetricsSubsystem::Get(Owner->GetWorld());
 	if (UNLIKELY(!Subsystem))
 	{
 		return 0;
@@ -187,7 +186,7 @@ int32 ReleaseMockMetricExtensions(
 
 void TestZeroState(UWorld* World)
 {
-	UWorldMetricsSubsystem* Subsystem = UE::WorldMetrics::GetSubsystem(World);
+	UWorldMetricsSubsystem* Subsystem = UWorldMetricsSubsystem::Get(World);
 	REQUIRE_MESSAGE(TEXT("WorldMetricsSubsystem is valid after GetSubsystem."), Subsystem);
 
 	REQUIRE_MESSAGE(TEXT("Unexpected existing metrics in WorldMetricsSubsystem"), !Subsystem->HasAnyMetric());
@@ -204,35 +203,9 @@ void TestSingleMetricAddRemove(UWorld* World)
 {
 	TestZeroState(World);
 
-	// UE::WorldMetrics API
-	{
-		REQUIRE_MESSAGE(TEXT("AddMetric failed"), UE::WorldMetrics::AddMetric<UMockWorldMetricA>(World));
-
-		UMockWorldMetricA* MetricA = UE::WorldMetrics::GetMetric<UMockWorldMetricA>(World);
-		REQUIRE_MESSAGE(TEXT("Test Metric initialization failed"), MetricA->InitializeCount == 1);
-
-		const UWorldMetricsSubsystem* Subsystem = UE::WorldMetrics::GetSubsystem(World);
-		REQUIRE_MESSAGE(
-			TEXT("WorldMetricsSubsystem should be enabled by default if there are metrics."), Subsystem->IsEnabled());
-		REQUIRE_MESSAGE(TEXT("WorldMetricsSubsystem should have metrics."), Subsystem->HasAnyMetric());
-		REQUIRE_MESSAGE(TEXT("WorldMetricsSubsystem should have one metric."), Subsystem->NumMetrics() == 1);
-
-		REQUIRE_MESSAGE(TEXT("Unexpected AddMetric"), !UE::WorldMetrics::AddMetric<UMockWorldMetricA>(World));
-		REQUIRE_MESSAGE(TEXT("WorldMetricsSubsystem should have one metric."), Subsystem->NumMetrics() == 1);
-
-		REQUIRE_MESSAGE(TEXT("RemoveMetric failed."), UE::WorldMetrics::RemoveMetric<UMockWorldMetricA>(World));
-		REQUIRE_MESSAGE(TEXT("Test Metric deinitialization failed after release."), MetricA->DeinitializeCount == 1);
-
-		TestZeroState(World);
-
-		REQUIRE_MESSAGE(TEXT("Unexpected RemoveMetric."), !UE::WorldMetrics::RemoveMetric<UMockWorldMetricA>(World));
-
-		TestZeroState(World);
-	}
-
 	// World Metrics Subsystem API
 	{
-		UWorldMetricsSubsystem* Subsystem = UE::WorldMetrics::GetSubsystem(World);
+		UWorldMetricsSubsystem* Subsystem = UWorldMetricsSubsystem::Get(World);
 
 		REQUIRE_MESSAGE(TEXT("AddMetric failed"), Subsystem->AddMetric<UMockWorldMetricA>());
 
@@ -266,7 +239,7 @@ void TestMultipleMetricsAddRemove(UWorld* World, bool bRandomized)
 
 	auto MetricClasses = Private::AddMockMetrics(World, NumMetrics, bRandomized);
 
-	const UWorldMetricsSubsystem* Subsystem = UE::WorldMetrics::GetSubsystem(World);
+	const UWorldMetricsSubsystem* Subsystem = UWorldMetricsSubsystem::Get(World);
 	REQUIRE_MESSAGE(
 		TEXT("WorldMetricsSubsystem should be enabled by default if there are metrics."), Subsystem->IsEnabled());
 	REQUIRE_MESSAGE(TEXT("WorldMetricsSubsystem should have extensions."), Subsystem->HasAnyMetric());
@@ -293,39 +266,32 @@ void TestSingleMetricSingleExtensionAcquireRelease(UWorld* World)
 {
 	TestZeroState(World);
 
-	REQUIRE_MESSAGE(TEXT("AddMetric failed"), UE::WorldMetrics::AddMetric<UMockWorldMetricA>(World));
+	UWorldMetricsSubsystem* Subsystem = UWorldMetricsSubsystem::Get(World);
+	REQUIRE_MESSAGE(TEXT("AddMetric failed"), Subsystem->AddMetric<UMockWorldMetricA>());
 
-	UMockWorldMetricA* MetricA = UE::WorldMetrics::GetMetric<UMockWorldMetricA>(World);
+	UMockWorldMetricA* MetricA = Subsystem->GetMetric<UMockWorldMetricA>();
 	REQUIRE_MESSAGE(TEXT("Test Metric initialization failed"), MetricA->InitializeCount == 1);
 
-	// UE::WorldMetrics API
+	// Extension acquire/release
 	{
-		UMockWorldMetricsExtensionA* ExtensionA =
-			UE::WorldMetrics::AcquireExtension<UMockWorldMetricsExtensionA>(MetricA);
-
+		UMockWorldMetricsExtensionA* ExtensionA = Subsystem->AcquireExtension<UMockWorldMetricsExtensionA>(MetricA);
 		REQUIRE_MESSAGE(TEXT("Acquire Test Extension failed."), ExtensionA);
 		REQUIRE_MESSAGE(TEXT("Test Extension initialization failed"), ExtensionA->InitializeCount == 1);
 
-		const UWorldMetricsSubsystem* Subsystem = UE::WorldMetrics::GetSubsystem(World);
 		REQUIRE_MESSAGE(
 			TEXT("WorldMetricsSubsystem should be enabled by default if there are metrics."), Subsystem->IsEnabled());
 		REQUIRE_MESSAGE(TEXT("WorldMetricsSubsystem should have extensions."), Subsystem->HasAnyExtension());
 		REQUIRE_MESSAGE(TEXT("WorldMetricsSubsystem should have one extension."), Subsystem->NumExtensions() == 1);
 
 		REQUIRE_MESSAGE(
-			TEXT("Test Extension release failed."),
-			UE::WorldMetrics::ReleaseExtension<UMockWorldMetricsExtensionA>(MetricA));
+			TEXT("Test Extension release failed."), Subsystem->ReleaseExtension<UMockWorldMetricsExtensionA>(MetricA));
 		REQUIRE_MESSAGE(
 			TEXT("Test Extension deinitialization failed after release."), ExtensionA->DeinitializeCount == 1);
-		REQUIRE_MESSAGE(TEXT("WorldMetricsSubsystem shouldn't be have any extension."), !Subsystem->HasAnyExtension());
-		REQUIRE_MESSAGE(TEXT("WorldMetricsSubsystem should have zero extensions."), Subsystem->NumExtensions() == 0);
 	}
 
-	// World Metrics Subsystem API
+	// Extension acquire/release through owner
 	{
-		UWorldMetricsSubsystem* Subsystem = UE::WorldMetrics::GetSubsystem(World);
-
-		UMockWorldMetricsExtensionA* ExtensionA = Subsystem->AcquireExtension<UMockWorldMetricsExtensionA>(MetricA);
+		UMockWorldMetricsExtensionA* ExtensionA = MetricA->GetOwner().AcquireExtension<UMockWorldMetricsExtensionA>(MetricA);
 		REQUIRE_MESSAGE(TEXT("Acquire Test Extension failed."), ExtensionA);
 		REQUIRE_MESSAGE(TEXT("Test Extension initialization failed"), ExtensionA->InitializeCount == 1);
 
@@ -335,12 +301,12 @@ void TestSingleMetricSingleExtensionAcquireRelease(UWorld* World)
 		REQUIRE_MESSAGE(TEXT("WorldMetricsSubsystem should have one extension."), Subsystem->NumExtensions() == 1);
 
 		REQUIRE_MESSAGE(
-			TEXT("Test Extension release failed."), Subsystem->ReleaseExtension<UMockWorldMetricsExtensionA>(MetricA));
+			TEXT("Test Extension release failed."), MetricA->GetOwner().ReleaseExtension<UMockWorldMetricsExtensionA>(MetricA));
 		REQUIRE_MESSAGE(
 			TEXT("Test Extension deinitialization failed after release."), ExtensionA->DeinitializeCount == 1);
 	}
 
-	UE::WorldMetrics::RemoveMetric<UMockWorldMetricA>(World);
+	Subsystem->RemoveMetric<UMockWorldMetricA>();
 
 	TestZeroState(World);
 }
@@ -355,15 +321,15 @@ void TestMultipleMetricsSingleExtensionAcquireRelease(UWorld* World, bool bRando
 	TArray<UMockWorldMetricBase*> MockMetrics = Private::GetMockMetrics(World, MetricClasses);
 	REQUIRE_MESSAGE(TEXT("Unexpected empty metrics list."), !MockMetrics.IsEmpty());
 
-	UMockWorldMetricsExtensionA* ExtensionA =
-		UE::WorldMetrics::AcquireExtension<UMockWorldMetricsExtensionA>(MockMetrics[0]);
+	UWorldMetricsSubsystem* Subsystem = UWorldMetricsSubsystem::Get(World);
+	UMockWorldMetricsExtensionA* ExtensionA = Subsystem->AcquireExtension<UMockWorldMetricsExtensionA>(MockMetrics[0]);
 
 	// Acquire extension loop
 	{
 		int32 CheckCount = 1;
 		for (UMockWorldMetricBase* MockMetric : MockMetrics)
 		{
-			UE::WorldMetrics::AcquireExtension<UMockWorldMetricsExtensionA>(MockMetric);
+			Subsystem->AcquireExtension<UMockWorldMetricsExtensionA>(MockMetric);
 			REQUIRE_MESSAGE(TEXT("Acquire Test Extension failed."), ExtensionA);
 			REQUIRE_MESSAGE(TEXT("Test Extension initialization failed"), ExtensionA->InitializeCount == 1);
 			REQUIRE_MESSAGE(TEXT("Test Extension initialization failed"), ExtensionA->OnAcquireCount == ++CheckCount);
@@ -380,7 +346,7 @@ void TestMultipleMetricsSingleExtensionAcquireRelease(UWorld* World, bool bRando
 		for (UMockWorldMetricBase* MockMetric : MockMetrics)
 		{
 			REQUIRE_MESSAGE(TEXT("Test Extension unexpected deinitialization"), ExtensionA->DeinitializeCount == 0);
-			UE::WorldMetrics::ReleaseExtension<UMockWorldMetricsExtensionA>(MockMetric);
+			Subsystem->ReleaseExtension<UMockWorldMetricsExtensionA>(MockMetric);
 
 			REQUIRE_MESSAGE(TEXT("Acquire Test Extension failed."), ExtensionA);
 			REQUIRE_MESSAGE(TEXT("Test Extension initialization failed"), ExtensionA->InitializeCount == 1);
@@ -414,7 +380,7 @@ void TestMultipleMetricsMultipleExtensionAcquireRelease(UWorld* World, bool bRan
 		Extensions.Emplace(Private::AcquireMockMetricExtensions(MockMetric, NumExtension, bRandomized));
 	}
 
-	const UWorldMetricsSubsystem* Subsystem = UE::WorldMetrics::GetSubsystem(World);
+	const UWorldMetricsSubsystem* Subsystem = UWorldMetricsSubsystem::Get(World);
 	REQUIRE_MESSAGE(
 		TEXT("WorldMetricsSubsystem should be enabled by default if there are metrics."), Subsystem->IsEnabled());
 	REQUIRE_MESSAGE(TEXT("WorldMetricsSubsystem should have extensions."), Subsystem->HasAnyExtension());
@@ -441,7 +407,7 @@ void TestMultipleMetricsEnable(UWorld* World, bool bRandomized)
 
 	auto MetricClasses = Private::AddMockMetrics(World, NumMetrics, bRandomized);
 
-	UWorldMetricsSubsystem* Subsystem = UE::WorldMetrics::GetSubsystem(World);
+	UWorldMetricsSubsystem* Subsystem = UWorldMetricsSubsystem::Get(World);
 	REQUIRE_MESSAGE(
 		TEXT("WorldMetricsSubsystem should be enabled by default if there are metrics."), Subsystem->IsEnabled());
 	REQUIRE_MESSAGE(TEXT("WorldMetricsSubsystem should have extensions."), Subsystem->HasAnyMetric());
@@ -480,21 +446,21 @@ void TestSingleMetricSingleExtensionAutoReleaseOnRemoval(UWorld* World)
 {
 	TestZeroState(World);
 
-	REQUIRE_MESSAGE(TEXT("AddMetric failed"), UE::WorldMetrics::AddMetric<UMockWorldMetricA>(World));
+	UWorldMetricsSubsystem* Subsystem = UWorldMetricsSubsystem::Get(World);
+	REQUIRE_MESSAGE(TEXT("AddMetric failed"), Subsystem->AddMetric<UMockWorldMetricA>());
 
-	UMockWorldMetricA* MetricA = UE::WorldMetrics::GetMetric<UMockWorldMetricA>(World);
+	UMockWorldMetricA* MetricA = Subsystem->GetMetric<UMockWorldMetricA>();
 
-	UMockWorldMetricsExtensionA* ExtensionA = UE::WorldMetrics::AcquireExtension<UMockWorldMetricsExtensionA>(MetricA);
+	UMockWorldMetricsExtensionA* ExtensionA = Subsystem->AcquireExtension<UMockWorldMetricsExtensionA>(MetricA);
 	REQUIRE_MESSAGE(TEXT("Acquire Test Extension failed."), ExtensionA);
 	REQUIRE_MESSAGE(TEXT("Test Extension initialization failed"), ExtensionA->InitializeCount == 1);
 
-	const UWorldMetricsSubsystem* Subsystem = UE::WorldMetrics::GetSubsystem(World);
 	REQUIRE_MESSAGE(
 		TEXT("WorldMetricsSubsystem should be enabled by default if there are metrics."), Subsystem->IsEnabled());
 	REQUIRE_MESSAGE(TEXT("WorldMetricsSubsystem should have extensions."), Subsystem->HasAnyExtension());
 	REQUIRE_MESSAGE(TEXT("WorldMetricsSubsystem should have one extension."), Subsystem->NumExtensions() == 1);
 
-	UE::WorldMetrics::RemoveMetric<UMockWorldMetricA>(World);
+	Subsystem->RemoveMetric<UMockWorldMetricA>();
 
 	REQUIRE_MESSAGE(TEXT("Test Extension deinitialization failed after release."), ExtensionA->DeinitializeCount == 1);
 	REQUIRE_MESSAGE(TEXT("WorldMetricsSubsystem shouldn't have any extensions."), !Subsystem->HasAnyExtension());
@@ -519,7 +485,7 @@ void TestMultipleMetricsMultipleExtensionAutoReleaseOnRemoval(UWorld* World, boo
 		Private::AcquireMockMetricExtensions(MockMetric, NumExtension, bRandomized);
 	}
 
-	const UWorldMetricsSubsystem* Subsystem = UE::WorldMetrics::GetSubsystem(World);
+	const UWorldMetricsSubsystem* Subsystem = UWorldMetricsSubsystem::Get(World);
 	REQUIRE_MESSAGE(
 		TEXT("WorldMetricsSubsystem should be enabled by default if there are metrics."), Subsystem->IsEnabled());
 	REQUIRE_MESSAGE(TEXT("WorldMetricsSubsystem should have extensions."), Subsystem->HasAnyExtension());
@@ -533,26 +499,27 @@ void TestSingleMetricOrphanExtensionAutoRemoval(UWorld* World)
 {
 	TestZeroState(World);
 
-	REQUIRE_MESSAGE(TEXT("AddMetric failed"), UE::WorldMetrics::AddMetric<UMockWorldMetricA>(World));
+	UWorldMetricsSubsystem* Subsystem = UWorldMetricsSubsystem::Get(World);
+	REQUIRE_MESSAGE(TEXT("AddMetric failed"), Subsystem->AddMetric<UMockWorldMetricA>());
 
-	UMockWorldMetricA* MetricA = UE::WorldMetrics::GetMetric<UMockWorldMetricA>(World);
+	UMockWorldMetricA* MetricA = Subsystem->GetMetric<UMockWorldMetricA>();
 
-	UMockWorldMetricsExtensionA* ExtensionA = UE::WorldMetrics::AcquireExtension<UMockWorldMetricsExtensionA>(MetricA);
+	UMockWorldMetricsExtensionA* ExtensionA = Subsystem->AcquireExtension<UMockWorldMetricsExtensionA>(MetricA);
 	REQUIRE_MESSAGE(TEXT("Acquire Test Extension failed."), ExtensionA);
 	REQUIRE_MESSAGE(TEXT("Test Extension initialization failed"), ExtensionA->InitializeCount == 1);
 
 	// Extension dependency: ExtensionB requires ExtensionA
-	UMockWorldMetricsExtensionB* ExtensionB = UE::WorldMetrics::AcquireExtension<UMockWorldMetricsExtensionB>(ExtensionA);
+	UMockWorldMetricsExtensionB* ExtensionB =
+		Subsystem->AcquireExtension<UMockWorldMetricsExtensionB>(ExtensionA);
 	REQUIRE_MESSAGE(TEXT("Acquire Test Extension failed."), ExtensionB);
 	REQUIRE_MESSAGE(TEXT("Test Extension initialization failed"), ExtensionB->InitializeCount == 1);
 
-	const UWorldMetricsSubsystem* Subsystem = UE::WorldMetrics::GetSubsystem(World);
 	REQUIRE_MESSAGE(
 		TEXT("WorldMetricsSubsystem should be enabled by default if there are metrics."), Subsystem->IsEnabled());
 	REQUIRE_MESSAGE(TEXT("WorldMetricsSubsystem should have extensions."), Subsystem->HasAnyExtension());
 	REQUIRE_MESSAGE(TEXT("WorldMetricsSubsystem should have one extension."), Subsystem->NumExtensions() == 2);
 
-	UE::WorldMetrics::RemoveMetric<UMockWorldMetricA>(World);
+	Subsystem->RemoveMetric<UMockWorldMetricA>();
 
 	REQUIRE_MESSAGE(TEXT("Test Extension deinitialization failed after release."), ExtensionA->DeinitializeCount == 1);
 	REQUIRE_MESSAGE(TEXT("Test Extension deinitialization failed after release."), ExtensionB->DeinitializeCount == 1);
@@ -563,7 +530,7 @@ void TestSingleMetricOrphanExtensionAutoRemoval(UWorld* World)
 
 void TestAll(UWorld* World)
 {
-	UWorldMetricsSubsystem* Subsystem = UE::WorldMetrics::GetSubsystem(World);
+	UWorldMetricsSubsystem* Subsystem = UWorldMetricsSubsystem::Get(World);
 	REQUIRE_MESSAGE(TEXT("WorldMetricsSubsystem is valid after GetSubsystem."), Subsystem);
 
 	Subsystem->Clear();
@@ -730,8 +697,7 @@ TEST_CASE_NAMED(
 	"[WorldMetrics][ClientContext][EngineFilter]")
 {
 	Private::ScopedWorldTest(
-		EWorldType::Editor,
-		[this](UWorld* World) { Private::TestSingleMetricOrphanExtensionAutoRemoval(World); });
+		EWorldType::Editor, [this](UWorld* World) { Private::TestSingleMetricOrphanExtensionAutoRemoval(World); });
 }
 
 }  // namespace UE::WorldMetrics
