WorldMetrics API update: added typed metric iteration
- Added ForEachMetricOfClass
- Added unit tests

#tests unit tests passed

diff --git a/Engine/Plugins/WorldMetrics/Source/WorldMetricsCore/Public/WorldMetricsSubsystem.h b/Engine/Plugins/WorldMetrics/Source/WorldMetricsCore/Public/WorldMetricsSubsystem.h
--- a/Engine/Plugins/WorldMetrics/Source/WorldMetricsCore/Public/WorldMetricsSubsystem.h
+++ b/Engine/Plugins/WorldMetrics/Source/WorldMetricsCore/Public/WorldMetricsSubsystem.h
@@ -140,6 +140,41 @@ public:
 	WORLDMETRICSCORE_API void ForEachMetric(const TFunctionRef<bool(const UWorldMetricInterface*)>& Func) const;
 	WORLDMETRICSCORE_API void ForEachMetric(const TFunctionRef<bool(UWorldMetricInterface*)>& Func);
 
+	/**
+	 * Invokes the parameter function on each metric of the template argument class contained by the subsystem.
+	 * @param Func The function which will be invoked for each metric. The function should return true to continue
+	 * execution, or false otherwise.
+	 */
+	template <typename MetricClass UE_REQUIRES(std::is_base_of_v<UWorldMetricInterface, MetricClass>)>
+	void ForEachMetricOfClass(const TFunctionRef<bool(const MetricClass*)>& Func) const
+	{
+		for (const UWorldMetricInterface* Metric : Metrics)
+		{
+			if (const MetricClass* TypedMetric = Cast<MetricClass>(Metric))
+			{
+				if (!Func(TypedMetric))
+				{
+					break;
+				}
+			}
+		}
+	}
+
+	template <typename MetricClass UE_REQUIRES(std::is_base_of_v<UWorldMetricInterface, MetricClass>)>
+	void ForEachMetricOfClass(const TFunctionRef<bool(MetricClass*)>& Func)
+	{
+		for (UWorldMetricInterface* Metric : Metrics)
+		{
+			if (MetricClass* TypedMetric = Cast<MetricClass>(Metric))
+			{
+				if (!Func(TypedMetric))
+				{
+					break;
+				}
+			}
+		}
+	}
+
 	/**
 	 * Acquires an extension on behalf of a object. The extension will be created the first time it is acquired,
 	 * and conversely it will be disabled once released by all owners.
diff --git a/Engine/Plugins/WorldMetrics/Source/WorldMetricsTest/Private/WorldMetricsTest.cpp b/Engine/Plugins/WorldMetrics/Source/WorldMetricsTest/Private/WorldMetricsTest.cpp
--- a/Engine/Plugins/WorldMetrics/Source/WorldMetricsTest/Private/WorldMetricsTest.cpp
+++ b/Engine/Plugins/WorldMetrics/Source/WorldMetricsTest/Private/WorldMetricsTest.cpp
@@ -242,11 +242,10 @@ void TestMultipleMetricsAddRemove(UWorld* World, bool bRandomized)
 	const UWorldMetricsSubsystem* Subsystem = UWorldMetricsSubsystem::Get(World);
 	REQUIRE_MESSAGE(
 		TEXT("WorldMetricsSubsystem should be enabled by default if there are metrics."), Subsystem->IsEnabled());
-	REQUIRE_MESSAGE(TEXT("WorldMetricsSubsystem should have extensions."), Subsystem->HasAnyMetric());
-	REQUIRE_MESSAGE(TEXT("WorldMetricsSubsystem should have one extension."), Subsystem->NumMetrics() == NumMetrics);
+	REQUIRE_MESSAGE(TEXT("WorldMetricsSubsystem should have metrics."), Subsystem->HasAnyMetric());
+	REQUIRE_MESSAGE(TEXT("WorldMetricsSubsystem should have one metric."), Subsystem->NumMetrics() == NumMetrics);
 
 	TArray<UMockWorldMetricBase*> MockMetrics = Private::GetMockMetrics(World, MetricClasses);
-	MockMetrics.Reserve(MetricClasses.Num());
 	for (const UMockWorldMetricBase* MockMetric : MockMetrics)
 	{
 		REQUIRE_MESSAGE(TEXT("Test Metric initialization failed"), MockMetric->InitializeCount == 1);
@@ -291,7 +290,8 @@ void TestSingleMetricSingleExtensionAcquireRelease(UWorld* World)
 
 	// Extension acquire/release through owner
 	{
-		UMockWorldMetricsExtensionA* ExtensionA = MetricA->GetOwner().AcquireExtension<UMockWorldMetricsExtensionA>(MetricA);
+		UMockWorldMetricsExtensionA* ExtensionA =
+			MetricA->GetOwner().AcquireExtension<UMockWorldMetricsExtensionA>(MetricA);
 		REQUIRE_MESSAGE(TEXT("Acquire Test Extension failed."), ExtensionA);
 		REQUIRE_MESSAGE(TEXT("Test Extension initialization failed"), ExtensionA->InitializeCount == 1);
 
@@ -301,7 +301,8 @@ void TestSingleMetricSingleExtensionAcquireRelease(UWorld* World)
 		REQUIRE_MESSAGE(TEXT("WorldMetricsSubsystem should have one extension."), Subsystem->NumExtensions() == 1);
 
 		REQUIRE_MESSAGE(
-			TEXT("Test Extension release failed."), MetricA->GetOwner().ReleaseExtension<UMockWorldMetricsExtensionA>(MetricA));
+			TEXT("Test Extension release failed."),
+			MetricA->GetOwner().ReleaseExtension<UMockWorldMetricsExtensionA>(MetricA));
 		REQUIRE_MESSAGE(
 			TEXT("Test Extension deinitialization failed after release."), ExtensionA->DeinitializeCount == 1);
 	}
@@ -410,11 +411,10 @@ void TestMultipleMetricsEnable(UWorld* World, bool bRandomized)
 	UWorldMetricsSubsystem* Subsystem = UWorldMetricsSubsystem::Get(World);
 	REQUIRE_MESSAGE(
 		TEXT("WorldMetricsSubsystem should be enabled by default if there are metrics."), Subsystem->IsEnabled());
-	REQUIRE_MESSAGE(TEXT("WorldMetricsSubsystem should have extensions."), Subsystem->HasAnyMetric());
-	REQUIRE_MESSAGE(TEXT("WorldMetricsSubsystem should have one extension."), Subsystem->NumMetrics() == NumMetrics);
+	REQUIRE_MESSAGE(TEXT("WorldMetricsSubsystem should have metrics."), Subsystem->HasAnyMetric());
+	REQUIRE_MESSAGE(TEXT("WorldMetricsSubsystem should have one metric."), Subsystem->NumMetrics() == NumMetrics);
 
 	TArray<UMockWorldMetricBase*> MockMetrics = Private::GetMockMetrics(World, MetricClasses);
-	MockMetrics.Reserve(MetricClasses.Num());
 	for (const UMockWorldMetricBase* MockMetric : MockMetrics)
 	{
 		REQUIRE_MESSAGE(TEXT("Test Metric initialization failed"), MockMetric->InitializeCount == 1);
@@ -509,8 +509,7 @@ void TestSingleMetricOrphanExtensionAutoRemoval(UWorld* World)
 	REQUIRE_MESSAGE(TEXT("Test Extension initialization failed"), ExtensionA->InitializeCount == 1);
 
 	// Extension dependency: ExtensionB requires ExtensionA
-	UMockWorldMetricsExtensionB* ExtensionB =
-		Subsystem->AcquireExtension<UMockWorldMetricsExtensionB>(ExtensionA);
+	UMockWorldMetricsExtensionB* ExtensionB = Subsystem->AcquireExtension<UMockWorldMetricsExtensionB>(ExtensionA);
 	REQUIRE_MESSAGE(TEXT("Acquire Test Extension failed."), ExtensionB);
 	REQUIRE_MESSAGE(TEXT("Test Extension initialization failed"), ExtensionB->InitializeCount == 1);
 
@@ -528,6 +527,101 @@ void TestSingleMetricOrphanExtensionAutoRemoval(UWorld* World)
 	TestZeroState(World);
 }
 
+void TestMultipleMetricIteration(UWorld* World, bool bRandomized)
+{
+	TestZeroState(World);
+
+	constexpr int32 NumMetrics = 8;
+
+	auto MetricClasses = Private::AddMockMetrics(World, NumMetrics, bRandomized);
+
+	// const iteration
+	{
+		const UWorldMetricsSubsystem* Subsystem = UWorldMetricsSubsystem::Get(World);
+		int32 CheckCount = 0;
+		int32 FooClassCount = 0;
+		int32 BarClassCount = 0;
+		Subsystem->ForEachMetric(
+			[&](const UWorldMetricInterface* Metric)
+			{
+				++CheckCount;
+				if (Metric->IsA<UMockWorldMetricFooBase>())
+				{
+					++FooClassCount;
+				}
+				if (Metric->IsA<UMockWorldMetricBarBase>())
+				{
+					++BarClassCount;
+				}
+				return true;
+			});
+		REQUIRE_MESSAGE(TEXT("Check count failed"), CheckCount == NumMetrics);
+
+		CheckCount = 0;
+		Subsystem->ForEachMetricOfClass<UMockWorldMetricFooBase>(
+			[&](const UMockWorldMetricFooBase* MockMetric)
+			{
+				++CheckCount;
+				return true;
+			});
+		REQUIRE_MESSAGE(TEXT("Check count failed"), CheckCount == FooClassCount);
+
+		CheckCount = 0;
+		Subsystem->ForEachMetricOfClass<UMockWorldMetricBarBase>(
+			[&](const UMockWorldMetricBarBase* MockMetric)
+			{
+				++CheckCount;
+				return true;
+			});
+		REQUIRE_MESSAGE(TEXT("Check count failed"), CheckCount == BarClassCount);
+	}
+
+	// non-const iteration
+	{
+		UWorldMetricsSubsystem* Subsystem = UWorldMetricsSubsystem::Get(World);
+		int32 CheckCount = 0;
+		int32 FooClassCount = 0;
+		int32 BarClassCount = 0;
+		Subsystem->ForEachMetric(
+			[&](UWorldMetricInterface* Metric)
+			{
+				++CheckCount;
+				if (Metric->IsA<UMockWorldMetricFooBase>())
+				{
+					++FooClassCount;
+				}
+				if (Metric->IsA<UMockWorldMetricBarBase>())
+				{
+					++BarClassCount;
+				}
+				return true;
+			});
+		REQUIRE_MESSAGE(TEXT("Check count failed"), CheckCount == NumMetrics);
+
+		CheckCount = 0;
+		Subsystem->ForEachMetricOfClass<UMockWorldMetricFooBase>(
+			[&](UMockWorldMetricFooBase* MockMetric)
+			{
+				++CheckCount;
+				return true;
+			});
+		REQUIRE_MESSAGE(TEXT("Check count failed"), CheckCount == FooClassCount);
+
+		CheckCount = 0;
+		Subsystem->ForEachMetricOfClass<UMockWorldMetricBarBase>(
+			[&](UMockWorldMetricBarBase* MockMetric)
+			{
+				++CheckCount;
+				return true;
+			});
+		REQUIRE_MESSAGE(TEXT("Check count failed"), CheckCount == BarClassCount);
+	}
+
+	const int32 NumMetricsRemoved = Private::RemoveMockMetrics(World, MetricClasses, bRandomized);
+
+	TestZeroState(World);
+}
+
 void TestAll(UWorld* World)
 {
 	UWorldMetricsSubsystem* Subsystem = UWorldMetricsSubsystem::Get(World);
@@ -543,6 +637,8 @@ void TestAll(UWorld* World)
 	TestMultipleMetricsAddRemove(World, true);
 	TestMultipleMetricsEnable(World, false);
 	TestMultipleMetricsEnable(World, true);
+	TestMultipleMetricIteration(World, false);
+	TestMultipleMetricIteration(World, true);
 
 	// Extension
 	TestSingleMetricSingleExtensionAcquireRelease(World);
@@ -630,6 +726,20 @@ TEST_CASE_NAMED(
 		});
 }
 
+TEST_CASE_NAMED(
+	WorldMetricsTestMultipleMetricIteration,
+	"WorldMetrics::TestMultipleMetricIteration",
+	"[WorldMetrics][ClientContext][EngineFilter]")
+{
+	Private::ScopedWorldTest(
+		EWorldType::Editor,
+		[this](UWorld* World)
+		{
+			Private::TestMultipleMetricIteration(World, false);
+			Private::TestMultipleMetricIteration(World, true);
+		});
+}
+
 TEST_CASE_NAMED(
 	WorldMetricsTestSingleMetricSingleExtensionAcquireRelease,
 	"WorldMetrics::TestSingleMetricSingleExtensionAcquireRelease",
diff --git a/Engine/Plugins/WorldMetrics/Source/WorldMetricsTest/Public/WorldMetricsTestTypes.h b/Engine/Plugins/WorldMetrics/Source/WorldMetricsTest/Public/WorldMetricsTestTypes.h
--- a/Engine/Plugins/WorldMetrics/Source/WorldMetricsTest/Public/WorldMetricsTestTypes.h
+++ b/Engine/Plugins/WorldMetrics/Source/WorldMetricsTest/Public/WorldMetricsTestTypes.h
@@ -51,61 +51,73 @@ public:
 };
 
 UCLASS()
-class UMockWorldMetricA : public UMockWorldMetricBase
+class UMockWorldMetricFooBase: public UMockWorldMetricBase
 {
 	GENERATED_BODY()
 };
 
 UCLASS()
-class UMockWorldMetricB : public UMockWorldMetricBase
+class UMockWorldMetricBarBase : public UMockWorldMetricBase
 {
 	GENERATED_BODY()
 };
 
 UCLASS()
-class UMockWorldMetricC : public UMockWorldMetricBase
+class UMockWorldMetricA : public UMockWorldMetricFooBase
 {
 	GENERATED_BODY()
 };
 
 UCLASS()
-class UMockWorldMetricD : public UMockWorldMetricBase
+class UMockWorldMetricB : public UMockWorldMetricFooBase
 {
 	GENERATED_BODY()
 };
 
 UCLASS()
-class UMockWorldMetricE : public UMockWorldMetricBase
+class UMockWorldMetricC : public UMockWorldMetricFooBase
 {
 	GENERATED_BODY()
 };
 
 UCLASS()
-class UMockWorldMetricF : public UMockWorldMetricBase
+class UMockWorldMetricD : public UMockWorldMetricFooBase
 {
 	GENERATED_BODY()
 };
 
 UCLASS()
-class UMockWorldMetricG : public UMockWorldMetricBase
+class UMockWorldMetricE : public UMockWorldMetricFooBase
 {
 	GENERATED_BODY()
 };
 
 UCLASS()
-class UMockWorldMetricH : public UMockWorldMetricBase
+class UMockWorldMetricF : public UMockWorldMetricBarBase
 {
 	GENERATED_BODY()
 };
 
 UCLASS()
-class UMockWorldMetricI : public UMockWorldMetricBase
+class UMockWorldMetricG : public UMockWorldMetricBarBase
 {
 	GENERATED_BODY()
 };
 
 UCLASS()
-class UMockWorldMetricJ : public UMockWorldMetricBase
+class UMockWorldMetricH : public UMockWorldMetricBarBase
+{
+	GENERATED_BODY()
+};
+
+UCLASS()
+class UMockWorldMetricI : public UMockWorldMetricBarBase
+{
+	GENERATED_BODY()
+};
+
+UCLASS()
+class UMockWorldMetricJ : public UMockWorldMetricBarBase
 {
 	GENERATED_BODY()
 };
