WorldMetrics API update: added checks against invalid metric and extension classes to UWorldMetricsSubsystem.
- Added unit test case for metric.
- Added unit test case for extension.

#tests unit tests passed

diff --git a/Engine/Plugins/WorldMetrics/Source/WorldMetricsCore/Private/WorldMetricsSubsystem.cpp b/Engine/Plugins/WorldMetrics/Source/WorldMetricsCore/Private/WorldMetricsSubsystem.cpp
--- a/Engine/Plugins/WorldMetrics/Source/WorldMetricsCore/Private/WorldMetricsSubsystem.cpp
+++ b/Engine/Plugins/WorldMetrics/Source/WorldMetricsCore/Private/WorldMetricsSubsystem.cpp
@@ -208,6 +208,12 @@ bool UWorldMetricsSubsystem::HasAnyMetric() const
 
 UWorldMetricInterface* UWorldMetricsSubsystem::GetMetric(const TSubclassOf<UWorldMetricInterface>& InMetricClass) const
 {
+	if (UNLIKELY(!InMetricClass))
+	{
+		UE_LOG(LogWorldMetrics, Error, TEXT("[%hs] Unexpected null metric class"), __FUNCTION__);
+		return nullptr;
+	}
+
 	if (InMetricClass->HasAnyClassFlags(CLASS_Abstract))
 	{
 		UE_LOG(
@@ -227,6 +233,12 @@ UWorldMetricInterface* UWorldMetricsSubsystem::GetMetric(const TSubclassOf<UWorl
 
 bool UWorldMetricsSubsystem::AddMetric(const TSubclassOf<UWorldMetricInterface>& InMetricClass)
 {
+	if (UNLIKELY(!InMetricClass))
+	{
+		UE_LOG(LogWorldMetrics, Error, TEXT("[%hs] Unexpected null metric class"), __FUNCTION__);
+		return false;
+	}
+
 	if (InMetricClass->HasAnyClassFlags(CLASS_Abstract))
 	{
 		UE_LOG(
@@ -269,6 +281,12 @@ bool UWorldMetricsSubsystem::AddMetric(const TSubclassOf<UWorldMetricInterface>&
 
 bool UWorldMetricsSubsystem::RemoveMetric(const TSubclassOf<UWorldMetricInterface>& InMetricClass)
 {
+	if (UNLIKELY(!InMetricClass))
+	{
+		UE_LOG(LogWorldMetrics, Error, TEXT("[%hs] Unexpected null metric class"), __FUNCTION__);
+		return false;
+	}
+
 	const int32 MetricIndex = GetMetricIndex(InMetricClass);
 	if (MetricIndex == INDEX_NONE)
 	{
@@ -403,6 +421,12 @@ UWorldMetricsExtension* UWorldMetricsSubsystem::AcquireExtensionInternal(
 		return nullptr;
 	}
 
+	if (UNLIKELY(!InExtensionClass))
+	{
+		UE_LOG(LogWorldMetrics, Error, TEXT("[%hs] Unexpected null extension class"), __FUNCTION__);
+		return nullptr;
+	}
+
 	if (InExtensionClass->HasAnyClassFlags(CLASS_Abstract))
 	{
 		UE_LOG(
@@ -479,6 +503,12 @@ bool UWorldMetricsSubsystem::ReleaseExtensionInternal(
 		return false;
 	}
 
+	if (UNLIKELY(!InExtensionClass))
+	{
+		UE_LOG(LogWorldMetrics, Error, TEXT("[%hs] Unexpected null extension class"), __FUNCTION__);
+		return false;
+	}
+
 	const int32 ExtensionIndex = GetExtensionIndex(InExtensionClass);
 	if (ExtensionIndex == INDEX_NONE)
 	{
diff --git a/Engine/Plugins/WorldMetrics/Source/WorldMetricsCore/Public/WorldMetricsSubsystem.h b/Engine/Plugins/WorldMetrics/Source/WorldMetricsCore/Public/WorldMetricsSubsystem.h
--- a/Engine/Plugins/WorldMetrics/Source/WorldMetricsCore/Public/WorldMetricsSubsystem.h
+++ b/Engine/Plugins/WorldMetrics/Source/WorldMetricsCore/Public/WorldMetricsSubsystem.h
@@ -241,7 +241,7 @@ private:
 	UPROPERTY(Transient)
 	TArray<TObjectPtr<UWorldMetricsExtension>> Extensions;
 
-	using FOwnerList = TSet<UObject*, DefaultKeyFuncs<UObject*>, TFixedSetAllocator<DefaultOwnerListCapacity>>;
+	using FOwnerList = TSet<UObject*, DefaultKeyFuncs<UObject*>, TInlineSetAllocator<DefaultOwnerListCapacity>>;
 	TArray<FOwnerList, TInlineAllocator<DefaultExtensionCapacity>> IndexedOwners;
 
 public:
diff --git a/Engine/Plugins/WorldMetrics/Source/WorldMetricsTest/Private/WorldMetricsTest.cpp b/Engine/Plugins/WorldMetrics/Source/WorldMetricsTest/Private/WorldMetricsTest.cpp
--- a/Engine/Plugins/WorldMetrics/Source/WorldMetricsTest/Private/WorldMetricsTest.cpp
+++ b/Engine/Plugins/WorldMetrics/Source/WorldMetricsTest/Private/WorldMetricsTest.cpp
@@ -207,8 +207,14 @@ void TestSingleMetricAddRemove(UWorld* World)
 	{
 		UWorldMetricsSubsystem* Subsystem = UWorldMetricsSubsystem::Get(World);
 
+		// Invalid metric add
+		REQUIRE_MESSAGE(TEXT("AddMetric should fail"), !Subsystem->AddMetric(UObject::StaticClass()));
+
 		REQUIRE_MESSAGE(TEXT("AddMetric failed"), Subsystem->AddMetric<UMockWorldMetricA>());
 
+		// Invalid metric get
+		REQUIRE_MESSAGE(TEXT("GetMetric should fail"), !Subsystem->GetMetric(UObject::StaticClass()));
+
 		UMockWorldMetricA* MetricA = Subsystem->GetMetric<UMockWorldMetricA>();
 		REQUIRE_MESSAGE(TEXT("Test Metric initialization failed"), MetricA->InitializeCount == 1);
 
@@ -220,6 +226,9 @@ void TestSingleMetricAddRemove(UWorld* World)
 		REQUIRE_MESSAGE(TEXT("Unexpected AddMetric"), !Subsystem->AddMetric<UMockWorldMetricA>());
 		REQUIRE_MESSAGE(TEXT("WorldMetricsSubsystem should have one metric."), Subsystem->NumMetrics() == 1);
 
+		// Invalid metric remove
+		REQUIRE_MESSAGE(TEXT("RemoveMetric should fail"), !Subsystem->RemoveMetric(UObject::StaticClass()));
+
 		REQUIRE_MESSAGE(TEXT("RemoveMetric failed."), Subsystem->RemoveMetric<UMockWorldMetricA>());
 		REQUIRE_MESSAGE(TEXT("Test Metric deinitialization failed after release."), MetricA->DeinitializeCount == 1);
 
@@ -273,6 +282,9 @@ void TestSingleMetricSingleExtensionAcquireRelease(UWorld* World)
 
 	// Extension acquire/release
 	{
+		// Invalid metric acquire
+		REQUIRE_MESSAGE(TEXT("AcquireExtension should fail"), !Subsystem->AcquireExtension(MetricA, UObject::StaticClass()));
+
 		UMockWorldMetricsExtensionA* ExtensionA = Subsystem->AcquireExtension<UMockWorldMetricsExtensionA>(MetricA);
 		REQUIRE_MESSAGE(TEXT("Acquire Test Extension failed."), ExtensionA);
 		REQUIRE_MESSAGE(TEXT("Test Extension initialization failed"), ExtensionA->InitializeCount == 1);
@@ -282,6 +294,9 @@ void TestSingleMetricSingleExtensionAcquireRelease(UWorld* World)
 		REQUIRE_MESSAGE(TEXT("WorldMetricsSubsystem should have extensions."), Subsystem->HasAnyExtension());
 		REQUIRE_MESSAGE(TEXT("WorldMetricsSubsystem should have one extension."), Subsystem->NumExtensions() == 1);
 
+		// Invalid metric release
+		REQUIRE_MESSAGE(TEXT("ReleaseExtension should fail"), !Subsystem->ReleaseExtension(MetricA, UObject::StaticClass()));
+
 		REQUIRE_MESSAGE(
 			TEXT("Test Extension release failed."), Subsystem->ReleaseExtension<UMockWorldMetricsExtensionA>(MetricA));
 		REQUIRE_MESSAGE(
