WorldMetrics API update: updated the World Metrics subsystem to implement a shared ownership model for metrics.
- Moved Enable to private to guarantee that the subsystem solely manages its internal metric update ticker.
- Added a CreateMetric factory method.
- Added ContainsMetric.
- Updated Add/RemoveMetric to accept a metric instance.
- Updated unit tests.
- Removed automated console command generation for metrics in favor of defining them by users.

#tests unit tests passed
#design https://docs.google.com/presentation/d/1btGRHfrjy2R4lL-2wI7CEdyIAQfQIqex_KQ83XsClb8/edit?usp=sharing

diff --git a/Engine/Plugins/WorldMetrics/Source/CsvMetrics/Private/CsvMetricsSubsystem.cpp b/Engine/Plugins/WorldMetrics/Source/CsvMetrics/Private/CsvMetricsSubsystem.cpp
--- a/Engine/Plugins/WorldMetrics/Source/CsvMetrics/Private/CsvMetricsSubsystem.cpp
+++ b/Engine/Plugins/WorldMetrics/Source/CsvMetrics/Private/CsvMetricsSubsystem.cpp
@@ -28,7 +28,7 @@ bool CanHaveCsvMetrics(const UWorld* World)
 bool UCsvMetricsSubsystem::ShouldCreateSubsystem(UObject* Outer) const
 {
 	// TODO: enable once CSVActorClassNameToCountMap is removed.
-	return false; // UE::CsvMetrics::Private::CanHaveCsvMetrics(Cast<UWorld>(Outer));
+	return false;  // UE::CsvMetrics::Private::CanHaveCsvMetrics(Cast<UWorld>(Outer));
 }
 
 void UCsvMetricsSubsystem::Initialize(FSubsystemCollectionBase& Collection)
@@ -42,6 +42,8 @@ void UCsvMetricsSubsystem::Deinitialize()
 {
 	UnbindProfilerCallbacks();
 
+	Metrics.Reset();
+
 	Super::Deinitialize();
 }
 
@@ -75,9 +77,10 @@ void UCsvMetricsSubsystem::AddMetrics()
 	UWorldMetricsSubsystem* WorldMetricsSubsystem = UWorldMetricsSubsystem::Get(GetWorld());
 	if (ensure(WorldMetricsSubsystem))
 	{
-		for (const TSubclassOf<UWorldMetricInterface>& MetricClass : Metrics)
+		check(Metrics.IsEmpty());
+		for (const TSubclassOf<UWorldMetricInterface>& MetricClass : MetricClasses)
 		{
-			WorldMetricsSubsystem->AddMetric(MetricClass);
+			Metrics.Emplace(WorldMetricsSubsystem->AddMetric(MetricClass));
 		}
 	}
 }
@@ -87,9 +90,10 @@ void UCsvMetricsSubsystem::RemoveMetrics()
 	UWorldMetricsSubsystem* WorldMetricsSubsystem = UWorldMetricsSubsystem::Get(GetWorld());
 	if (ensure(WorldMetricsSubsystem))
 	{
-		for (const TSubclassOf<UWorldMetricInterface>& MetricClass : Metrics)
+		for (UWorldMetricInterface* Metric : Metrics)
 		{
-			WorldMetricsSubsystem->RemoveMetric(MetricClass);
+			WorldMetricsSubsystem->RemoveMetric(Metric);
 		}
+		Metrics.Reset();
 	}
 }
diff --git a/Engine/Plugins/WorldMetrics/Source/CsvMetrics/Public/CsvMetricsSubsystem.h b/Engine/Plugins/WorldMetrics/Source/CsvMetrics/Public/CsvMetricsSubsystem.h
--- a/Engine/Plugins/WorldMetrics/Source/CsvMetrics/Public/CsvMetricsSubsystem.h
+++ b/Engine/Plugins/WorldMetrics/Source/CsvMetrics/Public/CsvMetricsSubsystem.h
@@ -27,9 +27,14 @@ public:
 
 	/** The CSV metric collection to add/remove when a CSV profiler capture stats/ends. */
 	UPROPERTY(Config)
-	TArray<TSubclassOf<UWorldMetricInterface>> Metrics;
+	TArray<TSubclassOf<UWorldMetricInterface>> MetricClasses;
 
 private:
+	/**
+	 * The collection of active metrics while a CSV profiler capture is running.
+	 * We rely on the World Metrics Subsystem's ownership for these so we don't keep hard-references. */
+	TArray<UWorldMetricInterface*> Metrics;
+
 	FDelegateHandle ProfileStartHandle;
 	FDelegateHandle ProfileEndHandle;
 
diff --git a/Engine/Plugins/WorldMetrics/Source/WorldMetricsCore/Private/WorldMetricsConsoleCommands.cpp b/Engine/Plugins/WorldMetrics/Source/WorldMetricsCore/Private/WorldMetricsConsoleCommands.cpp
deleted file mode 100644
diff --git a/Engine/Plugins/WorldMetrics/Source/WorldMetricsCore/Private/WorldMetricsConsoleCommands.h b/Engine/Plugins/WorldMetrics/Source/WorldMetricsCore/Private/WorldMetricsConsoleCommands.h
deleted file mode 100644
diff --git a/Engine/Plugins/WorldMetrics/Source/WorldMetricsCore/Private/WorldMetricsCoreModule.cpp b/Engine/Plugins/WorldMetrics/Source/WorldMetricsCore/Private/WorldMetricsCoreModule.cpp
--- a/Engine/Plugins/WorldMetrics/Source/WorldMetricsCore/Private/WorldMetricsCoreModule.cpp
+++ b/Engine/Plugins/WorldMetrics/Source/WorldMetricsCore/Private/WorldMetricsCoreModule.cpp
@@ -1,72 +1,5 @@
 // Copyright Epic Games, Inc. All Rights Reserved.
 
 #include "Modules/ModuleManager.h"
-#include "WorldMetricsDebug.h"
 
-#if WITH_WORLDMETRICS_DEBUG
-#include "HAL/IConsoleManager.h"
-#include "Misc/CoreDelegates.h"
-#include "WorldMetricsConsoleCommands.h"
-#endif	// WITH_WORLDMETRICS_DEBUG
-
-namespace UE::WorldMetrics
-{
-
-class FWorldMetricsCoreModule : public IModuleInterface
-{
-protected:
-#if WITH_WORLDMETRICS_DEBUG
-	//~ Begin IModuleInterface interface
-	virtual void StartupModule() override;
-	virtual void ShutdownModule() override;
-	//~ End IModuleInterface interface
-
-	void OnPostEngineInit();
-
-	void RegisterConsoleCommands();
-	void UnregisterConsoleCommands();
-	TArray<IConsoleObject*> ConsoleObjects;
-#endif	// WITH_WORLDMETRICS_DEBUG
-};
-
-#if WITH_WORLDMETRICS_DEBUG
-
-void FWorldMetricsCoreModule::StartupModule()
-{
-	FCoreDelegates::OnPostEngineInit.AddRaw(this, &FWorldMetricsCoreModule::OnPostEngineInit);
-}
-
-void FWorldMetricsCoreModule::ShutdownModule()
-{
-	UnregisterConsoleCommands();
-	FCoreDelegates::OnPostEngineInit.RemoveAll(this);
-}
-
-void FWorldMetricsCoreModule::OnPostEngineInit()
-{
-	RegisterConsoleCommands();
-}
-
-void FWorldMetricsCoreModule::RegisterConsoleCommands()
-{
-	IConsoleManager& ConsoleManager = IConsoleManager::Get();
-	RegisterCommonMetricConsoleCommands(ConsoleManager, ConsoleObjects);
-}
-
-void FWorldMetricsCoreModule::UnregisterConsoleCommands()
-{
-	IConsoleManager& ConsoleManager = IConsoleManager::Get();
-	for (IConsoleObject* ConsoleObject : ConsoleObjects)
-	{
-		if (ConsoleObject)
-		{
-			ConsoleManager.UnregisterConsoleObject(ConsoleObject);
-		}
-	}
-	ConsoleObjects.Reset();
-}
-#endif	// WITH_WORLDMETRICS_DEBUG
-
-}  // namespace UE::WorldMetrics
-
-IMPLEMENT_MODULE(UE::WorldMetrics::FWorldMetricsCoreModule, WorldMetricsCore)
+IMPLEMENT_MODULE(FDefaultModuleImpl, WorldMetricsCore)
diff --git a/Engine/Plugins/WorldMetrics/Source/WorldMetricsCore/Private/WorldMetricsSubsystem.cpp b/Engine/Plugins/WorldMetrics/Source/WorldMetricsCore/Private/WorldMetricsSubsystem.cpp
--- a/Engine/Plugins/WorldMetrics/Source/WorldMetricsCore/Private/WorldMetricsSubsystem.cpp
+++ b/Engine/Plugins/WorldMetrics/Source/WorldMetricsCore/Private/WorldMetricsSubsystem.cpp
@@ -12,20 +12,15 @@
 
 #include UE_INLINE_GENERATED_CPP_BY_NAME(WorldMetricsSubsystem)
 
-namespace UE::WorldMetrics::Private
-{
-
-bool CanHaveWorldMetrics(const UWorld* World)
-{
-	return World && World->IsGameWorld();
-}
-
-}  // namespace UE::WorldMetrics::Private
-
 //---------------------------------------------------------------------------------------------------------------------
 // UWorldMetricsSubsystem
 //---------------------------------------------------------------------------------------------------------------------
 
+bool UWorldMetricsSubsystem::CanHaveWorldMetrics(const UWorld* World)
+{
+	return World && (World->IsGameWorld() || World->WorldType == EWorldType::Editor);
+}
+
 UWorldMetricsSubsystem* UWorldMetricsSubsystem::Get(const UWorld* World)
 {
 	if (LIKELY(World))
@@ -37,7 +32,7 @@ UWorldMetricsSubsystem* UWorldMetricsSubsystem::Get(const UWorld* World)
 
 bool UWorldMetricsSubsystem::ShouldCreateSubsystem(UObject* Outer) const
 {
-	return UE::WorldMetrics::Private::CanHaveWorldMetrics(Cast<UWorld>(Outer));
+	return CanHaveWorldMetrics(Cast<UWorld>(Outer));
 }
 
 void UWorldMetricsSubsystem::Initialize(FSubsystemCollectionBase& Collection)
@@ -206,7 +201,7 @@ bool UWorldMetricsSubsystem::HasAnyMetric() const
 	return !Metrics.IsEmpty();
 }
 
-UWorldMetricInterface* UWorldMetricsSubsystem::GetMetric(const TSubclassOf<UWorldMetricInterface>& InMetricClass) const
+UWorldMetricInterface* UWorldMetricsSubsystem::CreateMetric(const TSubclassOf<UWorldMetricInterface>& InMetricClass)
 {
 	if (UNLIKELY(!InMetricClass))
 	{
@@ -222,50 +217,51 @@ UWorldMetricInterface* UWorldMetricsSubsystem::GetMetric(const TSubclassOf<UWorl
 		return nullptr;
 	}
 
-	const int32 MetricIndex = GetMetricIndex(InMetricClass);
-	if (MetricIndex == INDEX_NONE)
+	UWorldMetricInterface* Metric = NewObject<UWorldMetricInterface>(this, InMetricClass, NAME_None, RF_Transient);
+	if (UNLIKELY(!Metric))
 	{
-		return nullptr;
+		UE_LOG(
+			LogWorldMetrics, Error, TEXT("[%hs] Failed to create metric of class: %s"), __FUNCTION__,
+			*InMetricClass->GetFName().ToString());
 	}
-
-	return Metrics[MetricIndex];
+	return Metric;
 }
 
-bool UWorldMetricsSubsystem::AddMetric(const TSubclassOf<UWorldMetricInterface>& InMetricClass)
+bool UWorldMetricsSubsystem::ContainsMetric(UWorldMetricInterface* InMetric) const
 {
-	if (UNLIKELY(!InMetricClass))
+	if (UNLIKELY(!InMetric))
 	{
-		UE_LOG(LogWorldMetrics, Error, TEXT("[%hs] Unexpected null metric class"), __FUNCTION__);
+		UE_LOG(LogWorldMetrics, Error, TEXT("[%hs] Unexpected null metric instance"), __FUNCTION__);
 		return false;
 	}
+	return Algo::IndexOf(Metrics, InMetric) != INDEX_NONE;
+}
 
-	if (InMetricClass->HasAnyClassFlags(CLASS_Abstract))
-	{
-		UE_LOG(
-			LogWorldMetrics, Error, TEXT("[%hs] Parameter metric class is abstract: %s"), __FUNCTION__,
-			*InMetricClass->GetFName().ToString());
-		return false;
-	}
+UWorldMetricInterface* UWorldMetricsSubsystem::AddMetric(const TSubclassOf<UWorldMetricInterface>& InMetricClass)
+{
+	UWorldMetricInterface* Metric = CreateMetric(InMetricClass);
+	AddMetric(Metric);
+	return Metric;
+}
 
-	const int32 MetricIndex = GetMetricIndex(InMetricClass);
-	if (MetricIndex != INDEX_NONE)
+bool UWorldMetricsSubsystem::AddMetric(UWorldMetricInterface* InMetric)
+{
+	if (UNLIKELY(!InMetric))
 	{
+		UE_LOG(LogWorldMetrics, Error, TEXT("[%hs] Unexpected null metric instance"), __FUNCTION__);
 		return false;
 	}
 
-	UWorldMetricInterface* Metric = NewObject<UWorldMetricInterface>(this, InMetricClass, NAME_None, RF_Transient);
-	if (Metric == nullptr)
+	const int32 MetricIndex = Algo::IndexOf(Metrics, InMetric);
+	if (MetricIndex != INDEX_NONE)
 	{
-		UE_LOG(
-			LogWorldMetrics, Error, TEXT("[%hs] Failed to create metric of class: %s"), __FUNCTION__,
-			*InMetricClass->GetFName().ToString());
 		return false;
 	}
 
-	Metrics.Emplace(Metric);
+	Metrics.Emplace(InMetric);
 	if (IsEnabled())
 	{
-		Metric->Initialize();
+		InMetric->Initialize();
 	}
 	else
 	{
@@ -274,31 +270,29 @@ bool UWorldMetricsSubsystem::AddMetric(const TSubclassOf<UWorldMetricInterface>&
 
 	UE_LOG(
 		LogWorldMetrics, Log, TEXT("[%hs] Added metric of class %s."), __FUNCTION__,
-		*InMetricClass->GetFName().ToString());
+		*InMetric->GetClass()->GetFName().ToString());
 
 	return true;
 }
 
-bool UWorldMetricsSubsystem::RemoveMetric(const TSubclassOf<UWorldMetricInterface>& InMetricClass)
+bool UWorldMetricsSubsystem::RemoveMetric(UWorldMetricInterface* InMetric)
 {
-	if (UNLIKELY(!InMetricClass))
+	if (UNLIKELY(!InMetric))
 	{
-		UE_LOG(LogWorldMetrics, Error, TEXT("[%hs] Unexpected null metric class"), __FUNCTION__);
+		UE_LOG(LogWorldMetrics, Error, TEXT("[%hs] Unexpected null metric instance"), __FUNCTION__);
 		return false;
 	}
 
-	const int32 MetricIndex = GetMetricIndex(InMetricClass);
+	const int32 MetricIndex = Algo::IndexOf(Metrics, InMetric);
 	if (MetricIndex == INDEX_NONE)
 	{
 		return false;
 	}
 
-	UWorldMetricInterface* Metric = Metrics[MetricIndex];
 	if (IsEnabled())
 	{
-		Metric->Deinitialize();
-		VerifyMetricReleasedAllExtensions(Metric);
-		Metric->MarkAsGarbage();
+		InMetric->Deinitialize();
+		VerifyMetricReleasedAllExtensions(InMetric);
 	}
 	Metrics.RemoveAt(MetricIndex);
 	if (Metrics.IsEmpty())
@@ -309,7 +303,7 @@ bool UWorldMetricsSubsystem::RemoveMetric(const TSubclassOf<UWorldMetricInterfac
 
 	UE_LOG(
 		LogWorldMetrics, Log, TEXT("[%hs] Removed metric of class %s"), __FUNCTION__,
-		*InMetricClass->GetFName().ToString());
+		*InMetric->GetClass()->GetFName().ToString());
 
 	return true;
 }
@@ -342,17 +336,6 @@ void UWorldMetricsSubsystem::ForEachMetric(const TFunctionRef<bool(const UWorldM
 	}
 }
 
-void UWorldMetricsSubsystem::ForEachMetric(const TFunctionRef<bool(UWorldMetricInterface*)>& Func)
-{
-	for (UWorldMetricInterface* Metric : Metrics)
-	{
-		if (!Func(Metric))
-		{
-			break;
-		}
-	}
-}
-
 void UWorldMetricsSubsystem::OnUpdate(float DeltaTimeInSeconds)
 {
 	TRACE_CPUPROFILER_EVENT_SCOPE(UWorldMetricsSubsystem::OnUpdate);
@@ -370,12 +353,6 @@ void UWorldMetricsSubsystem::OnUpdate(float DeltaTimeInSeconds)
 	}
 }
 
-int32 UWorldMetricsSubsystem::GetMetricIndex(const TSubclassOf<UWorldMetricInterface>& InMetricClass) const
-{
-	return Algo::IndexOfByPredicate(
-		Metrics, [InMetricClass](UWorldMetricInterface* Metric) { return Metric->GetClass() == InMetricClass; });
-}
-
 UWorldMetricsExtension* UWorldMetricsSubsystem::AcquireExtension(
 	UWorldMetricInterface* InMetricOwner,
 	const TSubclassOf<UWorldMetricsExtension>& InExtensionClass)
diff --git a/Engine/Plugins/WorldMetrics/Source/WorldMetricsCore/Public/WorldMetricsDebug.h b/Engine/Plugins/WorldMetrics/Source/WorldMetricsCore/Public/WorldMetricsDebug.h
deleted file mode 100644
diff --git a/Engine/Plugins/WorldMetrics/Source/WorldMetricsCore/Public/WorldMetricsSubsystem.h b/Engine/Plugins/WorldMetrics/Source/WorldMetricsCore/Public/WorldMetricsSubsystem.h
--- a/Engine/Plugins/WorldMetrics/Source/WorldMetricsCore/Public/WorldMetricsSubsystem.h
+++ b/Engine/Plugins/WorldMetrics/Source/WorldMetricsCore/Public/WorldMetricsSubsystem.h
@@ -15,9 +15,15 @@ class UWorldMetricsSubsystem;
  * World metrics subsystem
  *
  * This subsystem provides an interface to add and remove world metrics implementing the UWorldMetricInterface class.
- * Both the subsystem and individual metrics can be enabled on demand. The World Metrics subsystem is the owner of all
- * the added metrics and these get automatically garbage collected after removal unless another system holds a
- * hard-reference to them.
+ *
+ * - Added metrics get automatically updated by the subsystem's ticker.
+ * - The subsystem becomes an owner of all added metrics. The user is responsible for removing them when no longer
+ *   needed so they can be garbage collected.
+ * - Metrics can have extensions to add shared functionality.
+ * - Extensions implement the UWorldMetricsExtension class and use Acquire/Release semantics. They can be acquired by
+ *   either metrics or extensions. Initialization and deinitialization are the ideal phases to do so.
+ * - The subsystem solely owns extensions and can automatically remove them for garbage collection whenever they are no
+ *   longer acquired by any metric or extension.
  */
 UCLASS(MinimalAPI, Config = Engine, DefaultConfig)
 class UWorldMetricsSubsystem final : public UWorldSubsystem
@@ -25,6 +31,7 @@ class UWorldMetricsSubsystem final : public UWorldSubsystem
 	GENERATED_BODY()
 
 public:
+	[[nodiscard]] WORLDMETRICSCORE_API static bool CanHaveWorldMetrics(const UWorld* World);
 	[[nodiscard]] WORLDMETRICSCORE_API static UWorldMetricsSubsystem* Get(const UWorld* World);
 
 	//~ Begin USubsystem
@@ -37,22 +44,11 @@ public:
 	WORLDMETRICSCORE_API virtual void BeginDestroy() override;
 	//~ End UObject
 
-	/**
-	 * Enables or disables the subsystem. When enabled the subsystem uses an update ticker to update each of the added
-	 * world metrics. All metrics are automatically initialized when the system is enabled and deinitialized when is
-	 * disabled. The subsystem is automatically disabled on deinitialization.
-	 *
-	 * @param bEnable: true for enabling the system and false otherwise.
-	 */
-	WORLDMETRICSCORE_API void Enable(bool bEnable);
-
 	[[nodiscard]] bool IsEnabled() const
 	{
 		return UpdateTickerHandle.IsValid();
 	}
 
-	WORLDMETRICSCORE_API void Clear();
-
 	/**
 	 * Sets the subsystem update ticker rate in seconds. This method automatically restarts the subsystem if it's
 	 * already enabled. Changing the update rate value on an enabled subsystem causes all metrics to be reinitialized.
@@ -83,65 +79,68 @@ public:
 	WORLDMETRICSCORE_API bool HasAnyExtension() const;
 
 	/**
-	 * Gets a world metric instance from the subsystem using the parameter metric
-	 * class.
+	 * Factory method to create world metric instances.
 	 *
 	 * @param InMetricClass: the class of the world metric to be retrieved.
-	 * @return A valid pointer to the world metric if found, and nullptr otherwise.
+	 * @return a valid pointer to a world metric if successfully created or nullptr otherwise.
 	 */
-	[[nodiscard]] WORLDMETRICSCORE_API UWorldMetricInterface* GetMetric(
-		const TSubclassOf<UWorldMetricInterface>& InMetricClass) const;
+	[[nodiscard]] WORLDMETRICSCORE_API UWorldMetricInterface* CreateMetric(
+		const TSubclassOf<UWorldMetricInterface>& InMetricClass);
 
 	template <typename MetricClass UE_REQUIRES(std::is_base_of_v<UWorldMetricInterface, MetricClass>)>
-	[[nodiscard]] MetricClass* GetMetric() const
+	[[nodiscard]] MetricClass* CreateMetric()
 	{
-		return static_cast<MetricClass*>(GetMetric(MetricClass::StaticClass()));
+		return static_cast<MetricClass*>(CreateMetric(MetricClass::StaticClass()));
 	}
 
 	/**
-	 * Adds a metric instance of the parameter class to the subsystem unless a metric of the same class and/or same ID
-	 * already exists. The system becomes the owner of all added metrics.
+	 * @param InMetric: the metric instance to search.
+	 * @return true if the parameter metric has been added or false otherwise.
+	 */
+	WORLDMETRICSCORE_API bool ContainsMetric(UWorldMetricInterface* InMetric) const;
+
+	/**
+	 * Adds a new metric instance of the parameter class to the subsystem.
 	 *
 	 * @param InMetricClass: the class of the metric to add.
-	 * @param InMetricId: the ID of the metric to add.
 	 * @return true if a metric of the parameter metric class is added and false otherwise.
 	 */
-	WORLDMETRICSCORE_API bool AddMetric(const TSubclassOf<UWorldMetricInterface>& InMetricClass);
+	[[nodiscard]] WORLDMETRICSCORE_API UWorldMetricInterface* AddMetric(
+		const TSubclassOf<UWorldMetricInterface>& InMetricClass);
 
 	template <typename MetricClass UE_REQUIRES(std::is_base_of_v<UWorldMetricInterface, MetricClass>)>
-	bool AddMetric()
+	[[nodiscard]] MetricClass* AddMetric()
 	{
-		return AddMetric(MetricClass::StaticClass());
+		return static_cast<MetricClass*>(AddMetric(MetricClass::StaticClass()));
 	}
 
 	/**
-	 * Removes the metric matching the parameter class and/or parameter ID from the subsystem.
-	 * The metric object gets automatically garbage-collected unless other system holds a hard-reference to it.
+	 * Adds the parameter metric instance to the subsystem. The subsystem will hold a hard-reference to the added
+	 * metric. The user is responsible for removing it when no longer needed so it can be garbage collected.
 	 *
-	 * @param InMetricClass: the class corresponding to the world metric to remove.
-	 * @param InMetricId: the ID corresponding to the world metric to remove.
-	 * @return true if a metric matching the parameter class is removed and false otherwise.
+	 * @param InMetric: the metric instance to add.
+	 * @return true if the parameter metric is valid and it wasn't previously added or false otherwise.
 	 */
-	WORLDMETRICSCORE_API bool RemoveMetric(const TSubclassOf<UWorldMetricInterface>& InMetricClass);
-
-	template <typename MetricClass UE_REQUIRES(std::is_base_of_v<UWorldMetricInterface, MetricClass>)>
-	bool RemoveMetric()
-	{
-		return RemoveMetric(MetricClass::StaticClass());
-	}
+	WORLDMETRICSCORE_API bool AddMetric(UWorldMetricInterface* InMetric);
 
-	WORLDMETRICSCORE_API void RemoveAllMetrics();
+	/**
+	 * Removes the parameter metric instance if it was previously added. The subsystem releases the hard-reference of
+	 * the metric removed so it can be garbage-collected.
+	 *
+	 * @param InMetric: the metric instance to remove.
+	 * @return true if the parameter metric is removed or false otherwise.
+	 */
+	WORLDMETRICSCORE_API bool RemoveMetric(UWorldMetricInterface* InMetric);
 
 	/**
-	 * Invokes the parameter function on each of the metrics contained by the subsystem.
+	 * Const iteration method for each of the metrics added to the subsystem.
 	 * @param Func The function which will be invoked for each metric. The function should return true to continue
 	 * execution, or false otherwise.
 	 */
 	WORLDMETRICSCORE_API void ForEachMetric(const TFunctionRef<bool(const UWorldMetricInterface*)>& Func) const;
-	WORLDMETRICSCORE_API void ForEachMetric(const TFunctionRef<bool(UWorldMetricInterface*)>& Func);
 
 	/**
-	 * Invokes the parameter function on each metric of the template argument class contained by the subsystem.
+	 * Const iteration method for each of the metrics added to the subsystem of the template argument class.
 	 * @param Func The function which will be invoked for each metric. The function should return true to continue
 	 * execution, or false otherwise.
 	 */
@@ -160,21 +159,6 @@ public:
 		}
 	}
 
-	template <typename MetricClass UE_REQUIRES(std::is_base_of_v<UWorldMetricInterface, MetricClass>)>
-	void ForEachMetricOfClass(const TFunctionRef<bool(MetricClass*)>& Func)
-	{
-		for (UWorldMetricInterface* Metric : Metrics)
-		{
-			if (MetricClass* TypedMetric = Cast<MetricClass>(Metric))
-			{
-				if (!Func(TypedMetric))
-				{
-					break;
-				}
-			}
-		}
-	}
-
 	/**
 	 * Acquires an extension on behalf of a object. The extension will be created the first time it is acquired,
 	 * and conversely it will be disabled once released by all owners.
@@ -258,15 +242,27 @@ private:
 	void InitializeMetrics();
 	void DeinitializeMetrics();
 
+	/**
+	 * Enables or disables the subsystem. When enabled the subsystem uses an update ticker to update each of the added
+	 * world metrics. All metrics are automatically initialized when the system is enabled and deinitialized when is
+	 * disabled. The subsystem is automatically disabled on deinitialization.
+	 *
+	 * @param bEnable: true for enabling the system and false otherwise.
+	 */
+	void Enable(bool bEnable);
+
 	void OnUpdate(float DeltaTimeInSeconds);
 
 	/**
-	 * Returns the live world metric list index corresponding to the parameter metric class.
-	 *
-	 * @param InMetricClass: the class type of the world metric.
-	 * @return a valid index value or INDEX_NONE if not matching metric exists.
+	 * Removes all added metrics. The subsystem releases the hard-reference of the  metrics removed so they can be
+	 * garbage-collected.
+	 */
+	void RemoveAllMetrics();
+
+	/**
+	 * Removes all existing metrics and extensions. Internal use only.
 	 */
-	int32 GetMetricIndex(const TSubclassOf<UWorldMetricInterface>& InMetricClass) const;
+	void Clear();
 
 	/**
 	 * Returns the live extension list index corresponding to the parameter extension class.
diff --git a/Engine/Plugins/WorldMetrics/Source/WorldMetricsTest/Private/WorldMetricsTest.cpp b/Engine/Plugins/WorldMetrics/Source/WorldMetricsTest/Private/WorldMetricsTest.cpp
--- a/Engine/Plugins/WorldMetrics/Source/WorldMetricsTest/Private/WorldMetricsTest.cpp
+++ b/Engine/Plugins/WorldMetrics/Source/WorldMetricsTest/Private/WorldMetricsTest.cpp
@@ -43,7 +43,7 @@ static const TSubclassOf<UMockWorldMetricsExtensionBase> MockExtensionClasses[]
 
 static constexpr int32 NumMockExtensionClasses = sizeof(MockExtensionClasses) / sizeof(MockExtensionClasses[0]);
 
-TArray<TSubclassOf<UMockWorldMetricBase>> AddMockMetrics(UWorld* World, int32 Count, bool bRandomized = false)
+TArray<TSubclassOf<UMockWorldMetricBase>> GetMockMetricClasses(UWorld* World, int32 Count, bool bRandomized = false)
 {
 	TArray<TSubclassOf<UMockWorldMetricBase>> Result;
 
@@ -60,15 +60,12 @@ TArray<TSubclassOf<UMockWorldMetricBase>> AddMockMetrics(UWorld* World, int32 Co
 
 	for (int32 ClassIndex : ClassIndices)
 	{
-		if (Subsystem->AddMetric(MockMetricClasses[ClassIndex]))
-		{
-			Result.Emplace(MockMetricClasses[ClassIndex]);
-		}
+		Result.Emplace(MockMetricClasses[ClassIndex]);
 	}
 	return Result;
 }
 
-TArray<UMockWorldMetricBase*> GetMockMetrics(
+TArray<UMockWorldMetricBase*> AddMockMetrics(
 	UWorld* World,
 	TArrayView<TSubclassOf<UMockWorldMetricBase>> InMockMetricClasses)
 {
@@ -83,18 +80,12 @@ TArray<UMockWorldMetricBase*> GetMockMetrics(
 	Result.Reserve(InMockMetricClasses.Num());
 	for (const TSubclassOf<UMockWorldMetricBase>& MetricClass : InMockMetricClasses)
 	{
-		if (UWorldMetricInterface* Metric = Subsystem->GetMetric(MetricClass))
-		{
-			Result.Emplace(static_cast<UMockWorldMetricBase*>(Metric));
-		}
+		Result.Emplace(static_cast<UMockWorldMetricBase*>(Subsystem->AddMetric(MetricClass)));
 	}
 	return Result;
 }
 
-int32 RemoveMockMetrics(
-	UWorld* World,
-	TArrayView<TSubclassOf<UMockWorldMetricBase>> InMockMetricClasses,
-	bool bRandomized = false)
+int32 RemoveMockMetrics(UWorld* World, TArrayView<UMockWorldMetricBase*> InMockMetrics, bool bRandomized = false)
 {
 	UWorldMetricsSubsystem* Subsystem = UWorldMetricsSubsystem::Get(World);
 	if (UNLIKELY(!Subsystem))
@@ -102,16 +93,16 @@ int32 RemoveMockMetrics(
 		return 0;
 	}
 
-	TArray<int32> ClassIndices = MakeIndexArray(InMockMetricClasses.Num());
+	TArray<int32> Indices = MakeIndexArray(InMockMetrics.Num());
 	if (bRandomized)
 	{
-		Shuffle(ClassIndices);
+		Shuffle(Indices);
 	}
 
 	int32 Result = 0;
-	for (int32 ClassIndex : ClassIndices)
+	for (int32 Index : Indices)
 	{
-		if (Subsystem->RemoveMetric(InMockMetricClasses[ClassIndex]))
+		if (Subsystem->RemoveMetric(InMockMetrics[Index]))
 		{
 			++Result;
 		}
@@ -189,10 +180,9 @@ void TestZeroState(UWorld* World)
 	UWorldMetricsSubsystem* Subsystem = UWorldMetricsSubsystem::Get(World);
 	REQUIRE_MESSAGE(TEXT("WorldMetricsSubsystem is valid after GetSubsystem."), Subsystem);
 
-	REQUIRE_MESSAGE(TEXT("Unexpected existing metrics in WorldMetricsSubsystem"), !Subsystem->HasAnyMetric());
-	Subsystem->Enable(true);
 	REQUIRE_MESSAGE(
 		TEXT("WorldMetricsSubsystem shouldn't be enabled if there are no metrics."), !Subsystem->IsEnabled());
+	REQUIRE_MESSAGE(TEXT("Unexpected existing metrics in WorldMetricsSubsystem"), !Subsystem->HasAnyMetric());
 	REQUIRE_MESSAGE(TEXT("WorldMetricsSubsystem should zero metrics."), Subsystem->NumMetrics() == 0);
 
 	REQUIRE_MESSAGE(TEXT("WorldMetricsSubsystem shouldn't be have any extension."), !Subsystem->HasAnyExtension());
@@ -203,41 +193,38 @@ void TestSingleMetricAddRemove(UWorld* World)
 {
 	TestZeroState(World);
 
-	// World Metrics Subsystem API
-	{
-		UWorldMetricsSubsystem* Subsystem = UWorldMetricsSubsystem::Get(World);
-
-		// Invalid metric add
-		REQUIRE_MESSAGE(TEXT("AddMetric should fail"), !Subsystem->AddMetric(UObject::StaticClass()));
-
-		REQUIRE_MESSAGE(TEXT("AddMetric failed"), Subsystem->AddMetric<UMockWorldMetricA>());
-
-		// Invalid metric get
-		REQUIRE_MESSAGE(TEXT("GetMetric should fail"), !Subsystem->GetMetric(UObject::StaticClass()));
+	UWorldMetricsSubsystem* Subsystem = UWorldMetricsSubsystem::Get(World);
 
-		UMockWorldMetricA* MetricA = Subsystem->GetMetric<UMockWorldMetricA>();
-		REQUIRE_MESSAGE(TEXT("Test Metric initialization failed"), MetricA->InitializeCount == 1);
+	// Invalid metric add
+	REQUIRE_MESSAGE(TEXT("AddMetric should fail"), !Subsystem->AddMetric(UObject::StaticClass()));
 
-		REQUIRE_MESSAGE(
-			TEXT("WorldMetricsSubsystem should be enabled by default if there are metrics."), Subsystem->IsEnabled());
-		REQUIRE_MESSAGE(TEXT("WorldMetricsSubsystem should have metrics."), Subsystem->HasAnyMetric());
-		REQUIRE_MESSAGE(TEXT("WorldMetricsSubsystem should have one metric."), Subsystem->NumMetrics() == 1);
+	UMockWorldMetricA* MetricA = Subsystem->AddMetric<UMockWorldMetricA>();
+	REQUIRE_MESSAGE(TEXT("ContainsMetric failed"), Subsystem->ContainsMetric(MetricA));
+	REQUIRE_MESSAGE(TEXT("Test Metric initialization failed"), MetricA->InitializeCount == 1);
 
-		REQUIRE_MESSAGE(TEXT("Unexpected AddMetric"), !Subsystem->AddMetric<UMockWorldMetricA>());
-		REQUIRE_MESSAGE(TEXT("WorldMetricsSubsystem should have one metric."), Subsystem->NumMetrics() == 1);
+	REQUIRE_MESSAGE(
+		TEXT("WorldMetricsSubsystem should be enabled by default if there are metrics."), Subsystem->IsEnabled());
+	REQUIRE_MESSAGE(TEXT("WorldMetricsSubsystem should have metrics."), Subsystem->HasAnyMetric());
+	REQUIRE_MESSAGE(TEXT("WorldMetricsSubsystem should have one metric."), Subsystem->NumMetrics() == 1);
 
-		// Invalid metric remove
-		REQUIRE_MESSAGE(TEXT("RemoveMetric should fail"), !Subsystem->RemoveMetric(UObject::StaticClass()));
+	// Same metric add
+	REQUIRE_MESSAGE(TEXT("Unexpected AddMetric"), !Subsystem->AddMetric(MetricA));
 
-		REQUIRE_MESSAGE(TEXT("RemoveMetric failed."), Subsystem->RemoveMetric<UMockWorldMetricA>());
-		REQUIRE_MESSAGE(TEXT("Test Metric deinitialization failed after release."), MetricA->DeinitializeCount == 1);
+	// Same metric class add
+	UMockWorldMetricA* AnotherMetricA = Subsystem->AddMetric<UMockWorldMetricA>();
+	REQUIRE_MESSAGE(TEXT("ContainsMetric failed"), Subsystem->ContainsMetric(AnotherMetricA));
+	REQUIRE_MESSAGE(TEXT("WorldMetricsSubsystem should have one metric."), Subsystem->NumMetrics() == 2);
 
-		TestZeroState(World);
+	// Metric remove
+	REQUIRE_MESSAGE(TEXT("RemoveMetric failed."), Subsystem->RemoveMetric(MetricA));
+	REQUIRE_MESSAGE(TEXT("Test Metric deinitialization failed after release."), MetricA->DeinitializeCount == 1);
+	REQUIRE_MESSAGE(TEXT("Unexpected RemoveMetric."), !Subsystem->RemoveMetric(MetricA));
 
-		REQUIRE_MESSAGE(TEXT("Unexpected RemoveMetric."), !Subsystem->RemoveMetric<UMockWorldMetricA>());
+	REQUIRE_MESSAGE(TEXT("Unexpected RemoveMetric."), Subsystem->RemoveMetric(AnotherMetricA));
+	REQUIRE_MESSAGE(TEXT("Test Metric deinitialization failed after release."), AnotherMetricA->DeinitializeCount == 1);
+	REQUIRE_MESSAGE(TEXT("Unexpected RemoveMetric."), !Subsystem->RemoveMetric(AnotherMetricA));
 
-		TestZeroState(World);
-	}
+	TestZeroState(World);
 }
 
 void TestMultipleMetricsAddRemove(UWorld* World, bool bRandomized)
@@ -246,21 +233,23 @@ void TestMultipleMetricsAddRemove(UWorld* World, bool bRandomized)
 
 	constexpr int32 NumMetrics = 5;
 
-	auto MetricClasses = Private::AddMockMetrics(World, NumMetrics, bRandomized);
-
 	const UWorldMetricsSubsystem* Subsystem = UWorldMetricsSubsystem::Get(World);
+
+	auto MetricClasses = Private::GetMockMetricClasses(World, NumMetrics, bRandomized);
+	TArray<UMockWorldMetricBase*> MockMetrics = Private::AddMockMetrics(World, MetricClasses);
+	REQUIRE_MESSAGE(TEXT("Unexpected empty metrics list."), !MockMetrics.IsEmpty());
+
 	REQUIRE_MESSAGE(
 		TEXT("WorldMetricsSubsystem should be enabled by default if there are metrics."), Subsystem->IsEnabled());
 	REQUIRE_MESSAGE(TEXT("WorldMetricsSubsystem should have metrics."), Subsystem->HasAnyMetric());
 	REQUIRE_MESSAGE(TEXT("WorldMetricsSubsystem should have one metric."), Subsystem->NumMetrics() == NumMetrics);
 
-	TArray<UMockWorldMetricBase*> MockMetrics = Private::GetMockMetrics(World, MetricClasses);
 	for (const UMockWorldMetricBase* MockMetric : MockMetrics)
 	{
 		REQUIRE_MESSAGE(TEXT("Test Metric initialization failed"), MockMetric->InitializeCount == 1);
 	}
 
-	const int32 NumMetricsRemoved = Private::RemoveMockMetrics(World, MetricClasses, bRandomized);
+	const int32 NumMetricsRemoved = Private::RemoveMockMetrics(World, MockMetrics, bRandomized);
 	REQUIRE_MESSAGE(TEXT("Unexpected metric remove count"), NumMetricsRemoved == NumMetrics);
 	for (const UMockWorldMetricBase* Metric : MockMetrics)
 	{
@@ -275,15 +264,14 @@ void TestSingleMetricSingleExtensionAcquireRelease(UWorld* World)
 	TestZeroState(World);
 
 	UWorldMetricsSubsystem* Subsystem = UWorldMetricsSubsystem::Get(World);
-	REQUIRE_MESSAGE(TEXT("AddMetric failed"), Subsystem->AddMetric<UMockWorldMetricA>());
-
-	UMockWorldMetricA* MetricA = Subsystem->GetMetric<UMockWorldMetricA>();
-	REQUIRE_MESSAGE(TEXT("Test Metric initialization failed"), MetricA->InitializeCount == 1);
+	UMockWorldMetricA* MetricA = Subsystem->AddMetric<UMockWorldMetricA>();
+	REQUIRE_MESSAGE(TEXT("ContainsMetric failed"), Subsystem->ContainsMetric(MetricA));
 
 	// Extension acquire/release
 	{
 		// Invalid metric acquire
-		REQUIRE_MESSAGE(TEXT("AcquireExtension should fail"), !Subsystem->AcquireExtension(MetricA, UObject::StaticClass()));
+		REQUIRE_MESSAGE(
+			TEXT("AcquireExtension should fail"), !Subsystem->AcquireExtension(MetricA, UObject::StaticClass()));
 
 		UMockWorldMetricsExtensionA* ExtensionA = Subsystem->AcquireExtension<UMockWorldMetricsExtensionA>(MetricA);
 		REQUIRE_MESSAGE(TEXT("Acquire Test Extension failed."), ExtensionA);
@@ -295,7 +283,8 @@ void TestSingleMetricSingleExtensionAcquireRelease(UWorld* World)
 		REQUIRE_MESSAGE(TEXT("WorldMetricsSubsystem should have one extension."), Subsystem->NumExtensions() == 1);
 
 		// Invalid metric release
-		REQUIRE_MESSAGE(TEXT("ReleaseExtension should fail"), !Subsystem->ReleaseExtension(MetricA, UObject::StaticClass()));
+		REQUIRE_MESSAGE(
+			TEXT("ReleaseExtension should fail"), !Subsystem->ReleaseExtension(MetricA, UObject::StaticClass()));
 
 		REQUIRE_MESSAGE(
 			TEXT("Test Extension release failed."), Subsystem->ReleaseExtension<UMockWorldMetricsExtensionA>(MetricA));
@@ -322,7 +311,7 @@ void TestSingleMetricSingleExtensionAcquireRelease(UWorld* World)
 			TEXT("Test Extension deinitialization failed after release."), ExtensionA->DeinitializeCount == 1);
 	}
 
-	Subsystem->RemoveMetric<UMockWorldMetricA>();
+	Subsystem->RemoveMetric(MetricA);
 
 	TestZeroState(World);
 }
@@ -333,8 +322,8 @@ void TestMultipleMetricsSingleExtensionAcquireRelease(UWorld* World, bool bRando
 
 	constexpr int32 NumMetrics = 5;
 
-	auto MetricClasses = Private::AddMockMetrics(World, NumMetrics, bRandomized);
-	TArray<UMockWorldMetricBase*> MockMetrics = Private::GetMockMetrics(World, MetricClasses);
+	auto MetricClasses = Private::GetMockMetricClasses(World, NumMetrics, bRandomized);
+	TArray<UMockWorldMetricBase*> MockMetrics = Private::AddMockMetrics(World, MetricClasses);
 	REQUIRE_MESSAGE(TEXT("Unexpected empty metrics list."), !MockMetrics.IsEmpty());
 
 	UWorldMetricsSubsystem* Subsystem = UWorldMetricsSubsystem::Get(World);
@@ -371,7 +360,7 @@ void TestMultipleMetricsSingleExtensionAcquireRelease(UWorld* World, bool bRando
 		REQUIRE_MESSAGE(TEXT("Test Extension initialization failed"), ExtensionA->DeinitializeCount == 1);
 	}
 
-	Private::RemoveMockMetrics(World, MetricClasses, bRandomized);
+	Private::RemoveMockMetrics(World, MockMetrics, bRandomized);
 
 	TestZeroState(World);
 }
@@ -383,8 +372,8 @@ void TestMultipleMetricsMultipleExtensionAcquireRelease(UWorld* World, bool bRan
 	constexpr int32 NumMetrics = 5;
 	constexpr int32 NumExtension = 5;
 
-	auto MetricClasses = Private::AddMockMetrics(World, NumMetrics, bRandomized);
-	TArray<UMockWorldMetricBase*> MockMetrics = Private::GetMockMetrics(World, MetricClasses);
+	auto MetricClasses = Private::GetMockMetricClasses(World, NumMetrics, bRandomized);
+	TArray<UMockWorldMetricBase*> MockMetrics = Private::AddMockMetrics(World, MetricClasses);
 	REQUIRE_MESSAGE(TEXT("Unexpected empty metrics list."), !MockMetrics.IsEmpty());
 
 	TArray<TArray<TSubclassOf<UMockWorldMetricsExtensionBase>>> Extensions;
@@ -410,49 +399,7 @@ void TestMultipleMetricsMultipleExtensionAcquireRelease(UWorld* World, bool bRan
 	REQUIRE_MESSAGE(TEXT("WorldMetricsSubsystem shouldn't be have any extension."), !Subsystem->HasAnyExtension());
 	REQUIRE_MESSAGE(TEXT("WorldMetricsSubsystem should have zero extensions."), Subsystem->NumExtensions() == 0);
 
-	Private::RemoveMockMetrics(World, MetricClasses, bRandomized);
-
-	TestZeroState(World);
-}
-
-void TestMultipleMetricsEnable(UWorld* World, bool bRandomized)
-{
-	TestZeroState(World);
-
-	constexpr int32 NumMetrics = 5;
-
-	auto MetricClasses = Private::AddMockMetrics(World, NumMetrics, bRandomized);
-
-	UWorldMetricsSubsystem* Subsystem = UWorldMetricsSubsystem::Get(World);
-	REQUIRE_MESSAGE(
-		TEXT("WorldMetricsSubsystem should be enabled by default if there are metrics."), Subsystem->IsEnabled());
-	REQUIRE_MESSAGE(TEXT("WorldMetricsSubsystem should have metrics."), Subsystem->HasAnyMetric());
-	REQUIRE_MESSAGE(TEXT("WorldMetricsSubsystem should have one metric."), Subsystem->NumMetrics() == NumMetrics);
-
-	TArray<UMockWorldMetricBase*> MockMetrics = Private::GetMockMetrics(World, MetricClasses);
-	for (const UMockWorldMetricBase* MockMetric : MockMetrics)
-	{
-		REQUIRE_MESSAGE(TEXT("Test Metric initialization failed"), MockMetric->InitializeCount == 1);
-	}
-
-	Subsystem->Enable(false);
-	for (const UMockWorldMetricBase* MockMetric : MockMetrics)
-	{
-		REQUIRE_MESSAGE(TEXT("Test Metric initialization failed"), MockMetric->DeinitializeCount == 1);
-	}
-
-	Subsystem->Enable(true);
-	for (const UMockWorldMetricBase* MockMetric : MockMetrics)
-	{
-		REQUIRE_MESSAGE(TEXT("Test Metric initialization failed"), MockMetric->InitializeCount == 2);
-	}
-
-	const int32 NumMetricsRemoved = Private::RemoveMockMetrics(World, MetricClasses, bRandomized);
-	REQUIRE_MESSAGE(TEXT("Unexpected metric remove count"), NumMetricsRemoved == NumMetrics);
-	for (const UMockWorldMetricBase* Metric : MockMetrics)
-	{
-		REQUIRE_MESSAGE(TEXT("Test Metric initialization failed"), Metric->DeinitializeCount == 2);
-	}
+	Private::RemoveMockMetrics(World, MockMetrics, bRandomized);
 
 	TestZeroState(World);
 }
@@ -462,9 +409,8 @@ void TestSingleMetricSingleExtensionAutoReleaseOnRemoval(UWorld* World)
 	TestZeroState(World);
 
 	UWorldMetricsSubsystem* Subsystem = UWorldMetricsSubsystem::Get(World);
-	REQUIRE_MESSAGE(TEXT("AddMetric failed"), Subsystem->AddMetric<UMockWorldMetricA>());
-
-	UMockWorldMetricA* MetricA = Subsystem->GetMetric<UMockWorldMetricA>();
+	UMockWorldMetricA* MetricA = Subsystem->AddMetric<UMockWorldMetricA>();
+	REQUIRE_MESSAGE(TEXT("ContainsMetric failed"), Subsystem->ContainsMetric(MetricA));
 
 	UMockWorldMetricsExtensionA* ExtensionA = Subsystem->AcquireExtension<UMockWorldMetricsExtensionA>(MetricA);
 	REQUIRE_MESSAGE(TEXT("Acquire Test Extension failed."), ExtensionA);
@@ -475,7 +421,7 @@ void TestSingleMetricSingleExtensionAutoReleaseOnRemoval(UWorld* World)
 	REQUIRE_MESSAGE(TEXT("WorldMetricsSubsystem should have extensions."), Subsystem->HasAnyExtension());
 	REQUIRE_MESSAGE(TEXT("WorldMetricsSubsystem should have one extension."), Subsystem->NumExtensions() == 1);
 
-	Subsystem->RemoveMetric<UMockWorldMetricA>();
+	Subsystem->RemoveMetric(MetricA);
 
 	REQUIRE_MESSAGE(TEXT("Test Extension deinitialization failed after release."), ExtensionA->DeinitializeCount == 1);
 	REQUIRE_MESSAGE(TEXT("WorldMetricsSubsystem shouldn't have any extensions."), !Subsystem->HasAnyExtension());
@@ -490,8 +436,8 @@ void TestMultipleMetricsMultipleExtensionAutoReleaseOnRemoval(UWorld* World, boo
 	constexpr int32 NumMetrics = 5;
 	constexpr int32 NumExtension = 5;
 
-	auto MetricClasses = Private::AddMockMetrics(World, NumMetrics, bRandomized);
-	TArray<UMockWorldMetricBase*> MockMetrics = Private::GetMockMetrics(World, MetricClasses);
+	auto MetricClasses = Private::GetMockMetricClasses(World, NumMetrics, bRandomized);
+	TArray<UMockWorldMetricBase*> MockMetrics = Private::AddMockMetrics(World, MetricClasses);
 	REQUIRE_MESSAGE(TEXT("Unexpected empty metrics list."), !MockMetrics.IsEmpty());
 
 	// Acquire extensions
@@ -505,7 +451,7 @@ void TestMultipleMetricsMultipleExtensionAutoReleaseOnRemoval(UWorld* World, boo
 		TEXT("WorldMetricsSubsystem should be enabled by default if there are metrics."), Subsystem->IsEnabled());
 	REQUIRE_MESSAGE(TEXT("WorldMetricsSubsystem should have extensions."), Subsystem->HasAnyExtension());
 
-	Private::RemoveMockMetrics(World, MetricClasses, bRandomized);
+	Private::RemoveMockMetrics(World, MockMetrics, bRandomized);
 
 	TestZeroState(World);
 }
@@ -515,9 +461,8 @@ void TestSingleMetricOrphanExtensionAutoRemoval(UWorld* World)
 	TestZeroState(World);
 
 	UWorldMetricsSubsystem* Subsystem = UWorldMetricsSubsystem::Get(World);
-	REQUIRE_MESSAGE(TEXT("AddMetric failed"), Subsystem->AddMetric<UMockWorldMetricA>());
-
-	UMockWorldMetricA* MetricA = Subsystem->GetMetric<UMockWorldMetricA>();
+	UMockWorldMetricA* MetricA = Subsystem->AddMetric<UMockWorldMetricA>();
+	REQUIRE_MESSAGE(TEXT("ContainsMetric failed"), Subsystem->ContainsMetric(MetricA));
 
 	UMockWorldMetricsExtensionA* ExtensionA = Subsystem->AcquireExtension<UMockWorldMetricsExtensionA>(MetricA);
 	REQUIRE_MESSAGE(TEXT("Acquire Test Extension failed."), ExtensionA);
@@ -533,7 +478,7 @@ void TestSingleMetricOrphanExtensionAutoRemoval(UWorld* World)
 	REQUIRE_MESSAGE(TEXT("WorldMetricsSubsystem should have extensions."), Subsystem->HasAnyExtension());
 	REQUIRE_MESSAGE(TEXT("WorldMetricsSubsystem should have one extension."), Subsystem->NumExtensions() == 2);
 
-	Subsystem->RemoveMetric<UMockWorldMetricA>();
+	Subsystem->RemoveMetric(MetricA);
 
 	REQUIRE_MESSAGE(TEXT("Test Extension deinitialization failed after release."), ExtensionA->DeinitializeCount == 1);
 	REQUIRE_MESSAGE(TEXT("Test Extension deinitialization failed after release."), ExtensionB->DeinitializeCount == 1);
@@ -548,7 +493,9 @@ void TestMultipleMetricIteration(UWorld* World, bool bRandomized)
 
 	constexpr int32 NumMetrics = 8;
 
-	auto MetricClasses = Private::AddMockMetrics(World, NumMetrics, bRandomized);
+	auto MetricClasses = Private::GetMockMetricClasses(World, NumMetrics, bRandomized);
+	TArray<UMockWorldMetricBase*> MockMetrics = Private::AddMockMetrics(World, MetricClasses);
+	REQUIRE_MESSAGE(TEXT("Unexpected empty metrics list."), !MockMetrics.IsEmpty());
 
 	// const iteration
 	{
@@ -591,48 +538,7 @@ void TestMultipleMetricIteration(UWorld* World, bool bRandomized)
 		REQUIRE_MESSAGE(TEXT("Check count failed"), CheckCount == BarClassCount);
 	}
 
-	// non-const iteration
-	{
-		UWorldMetricsSubsystem* Subsystem = UWorldMetricsSubsystem::Get(World);
-		int32 CheckCount = 0;
-		int32 FooClassCount = 0;
-		int32 BarClassCount = 0;
-		Subsystem->ForEachMetric(
-			[&](UWorldMetricInterface* Metric)
-			{
-				++CheckCount;
-				if (Metric->IsA<UMockWorldMetricFooBase>())
-				{
-					++FooClassCount;
-				}
-				if (Metric->IsA<UMockWorldMetricBarBase>())
-				{
-					++BarClassCount;
-				}
-				return true;
-			});
-		REQUIRE_MESSAGE(TEXT("Check count failed"), CheckCount == NumMetrics);
-
-		CheckCount = 0;
-		Subsystem->ForEachMetricOfClass<UMockWorldMetricFooBase>(
-			[&](UMockWorldMetricFooBase* MockMetric)
-			{
-				++CheckCount;
-				return true;
-			});
-		REQUIRE_MESSAGE(TEXT("Check count failed"), CheckCount == FooClassCount);
-
-		CheckCount = 0;
-		Subsystem->ForEachMetricOfClass<UMockWorldMetricBarBase>(
-			[&](UMockWorldMetricBarBase* MockMetric)
-			{
-				++CheckCount;
-				return true;
-			});
-		REQUIRE_MESSAGE(TEXT("Check count failed"), CheckCount == BarClassCount);
-	}
-
-	const int32 NumMetricsRemoved = Private::RemoveMockMetrics(World, MetricClasses, bRandomized);
+	Private::RemoveMockMetrics(World, MockMetrics, bRandomized);
 
 	TestZeroState(World);
 }
@@ -642,16 +548,12 @@ void TestAll(UWorld* World)
 	UWorldMetricsSubsystem* Subsystem = UWorldMetricsSubsystem::Get(World);
 	REQUIRE_MESSAGE(TEXT("WorldMetricsSubsystem is valid after GetSubsystem."), Subsystem);
 
-	Subsystem->Clear();
-
 	TestZeroState(World);
 
 	// Metrics
 	TestSingleMetricAddRemove(World);
 	TestMultipleMetricsAddRemove(World, false);
 	TestMultipleMetricsAddRemove(World, true);
-	TestMultipleMetricsEnable(World, false);
-	TestMultipleMetricsEnable(World, true);
 	TestMultipleMetricIteration(World, false);
 	TestMultipleMetricIteration(World, true);
 
@@ -727,20 +629,6 @@ TEST_CASE_NAMED(
 		});
 }
 
-TEST_CASE_NAMED(
-	WorldMetricsTestMultipleMetricsEnable,
-	"WorldMetrics::TestMultipleMetricsEnable",
-	"[WorldMetrics][ClientContext][EngineFilter]")
-{
-	Private::ScopedWorldTest(
-		EWorldType::Editor,
-		[this](UWorld* World)
-		{
-			Private::TestMultipleMetricsEnable(World, false);
-			Private::TestMultipleMetricsEnable(World, true);
-		});
-}
-
 TEST_CASE_NAMED(
 	WorldMetricsTestMultipleMetricIteration,
 	"WorldMetrics::TestMultipleMetricIteration",
