WorldMetrics API update: removed client context from unit tests.

diff --git a/Engine/Plugins/WorldMetrics/Source/WorldMetricsTest/Private/WorldMetricsTest.cpp b/Engine/Plugins/WorldMetrics/Source/WorldMetricsTest/Private/WorldMetricsTest.cpp
--- a/Engine/Plugins/WorldMetrics/Source/WorldMetricsTest/Private/WorldMetricsTest.cpp
+++ b/Engine/Plugins/WorldMetrics/Source/WorldMetricsTest/Private/WorldMetricsTest.cpp
@@ -7,7 +7,6 @@
 #endif	// WITH_TESTS
 
 #include "HAL/IConsoleManager.h"
-#include "WorldMetricsLog.h"
 #include "WorldMetricsSubsystem.h"
 #include "WorldMetricsTestTypes.h"
 #include "WorldMetricsTestUtil.h"
@@ -602,7 +601,7 @@ static void ScopedWorldTest(EWorldType::Type WorldType, TFunctionRef<void(UWorld
 }
 }  // namespace Private
 
-TEST_CASE_NAMED(WorldMetricsTestZeroState, "WorldMetrics::TestZeroState", "[WorldMetrics][ClientContext][EngineFilter]")
+TEST_CASE_NAMED(WorldMetricsTestZeroState, "WorldMetrics::TestZeroState", "[WorldMetrics][EngineFilter]")
 {
 	Private::ScopedWorldTest(EWorldType::Editor, [this](UWorld* World) { Private::TestZeroState(World); });
 }
@@ -610,7 +609,7 @@ TEST_CASE_NAMED(WorldMetricsTestZeroState, "WorldMetrics::TestZeroState", "[Worl
 TEST_CASE_NAMED(
 	WorldMetricsTestSingleMetricAddRemove,
 	"WorldMetrics::TestSingleMetricAddRemove",
-	"[WorldMetrics][ClientContext][EngineFilter]")
+	"[WorldMetrics][EngineFilter]")
 {
 	Private::ScopedWorldTest(EWorldType::Editor, [this](UWorld* World) { Private::TestSingleMetricAddRemove(World); });
 }
@@ -618,7 +617,7 @@ TEST_CASE_NAMED(
 TEST_CASE_NAMED(
 	WorldMetricsTestMultipleMetricsAddRemove,
 	"WorldMetrics::TestMultipleMetricsAddRemove",
-	"[WorldMetrics][ClientContext][EngineFilter]")
+	"[WorldMetrics][EngineFilter]")
 {
 	Private::ScopedWorldTest(
 		EWorldType::Editor,
@@ -632,7 +631,7 @@ TEST_CASE_NAMED(
 TEST_CASE_NAMED(
 	WorldMetricsTestMultipleMetricIteration,
 	"WorldMetrics::TestMultipleMetricIteration",
-	"[WorldMetrics][ClientContext][EngineFilter]")
+	"[WorldMetrics][EngineFilter]")
 {
 	Private::ScopedWorldTest(
 		EWorldType::Editor,
@@ -646,7 +645,7 @@ TEST_CASE_NAMED(
 TEST_CASE_NAMED(
 	WorldMetricsTestSingleMetricSingleExtensionAcquireRelease,
 	"WorldMetrics::TestSingleMetricSingleExtensionAcquireRelease",
-	"[WorldMetrics][ClientContext][EngineFilter]")
+	"[WorldMetrics][EngineFilter]")
 {
 	Private::ScopedWorldTest(
 		EWorldType::Editor, [this](UWorld* World) { Private::TestSingleMetricSingleExtensionAcquireRelease(World); });
@@ -655,7 +654,7 @@ TEST_CASE_NAMED(
 TEST_CASE_NAMED(
 	WorldMetricsTestMultipleMetricsSingleExtensionAcquireRelease,
 	"WorldMetrics::TestMultipleMetricsSingleExtensionAcquireRelease",
-	"[WorldMetrics][ClientContext][EngineFilter]")
+	"[WorldMetrics][EngineFilter]")
 {
 	Private::ScopedWorldTest(
 		EWorldType::Editor,
@@ -669,7 +668,7 @@ TEST_CASE_NAMED(
 TEST_CASE_NAMED(
 	WorldMetricsTestMultipleMetricsMultipleExtensionAcquireRelease,
 	"WorldMetrics::TestMultipleMetricsMultipleExtensionAcquireRelease",
-	"[WorldMetrics][ClientContext][EngineFilter]")
+	"[WorldMetrics][EngineFilter]")
 {
 	Private::ScopedWorldTest(
 		EWorldType::Editor,
@@ -683,7 +682,7 @@ TEST_CASE_NAMED(
 TEST_CASE_NAMED(
 	WorldMetricsTestSingleMetricSingleExtensionAutoReleaseOnRemoval,
 	"WorldMetrics::TestSingleMetricSingleExtensionAutoReleaseOnRemoval",
-	"[WorldMetrics][ClientContext][EngineFilter]")
+	"[WorldMetrics][EngineFilter]")
 {
 	Private::ScopedWorldTest(
 		EWorldType::Editor,
@@ -693,7 +692,7 @@ TEST_CASE_NAMED(
 TEST_CASE_NAMED(
 	WorldMetricsTestMultipleMetricsMultipleExtensionAutoReleaseOnRemoval,
 	"WorldMetrics::TestMultipleMetricsMultipleExtensionAutoReleaseOnRemoval",
-	"[WorldMetrics][ClientContext][EngineFilter]")
+	"[WorldMetrics][EngineFilter]")
 {
 	Private::ScopedWorldTest(
 		EWorldType::Editor,
@@ -707,7 +706,7 @@ TEST_CASE_NAMED(
 TEST_CASE_NAMED(
 	WorldMetricsTestSingleMetricOrphanExtensionAutoRemoval,
 	"WorldMetrics::TestSingleMetricOrphanExtensionAutoRemoval",
-	"[WorldMetrics][ClientContext][EngineFilter]")
+	"[WorldMetrics][EngineFilter]")
 {
 	Private::ScopedWorldTest(
 		EWorldType::Editor, [this](UWorld* World) { Private::TestSingleMetricOrphanExtensionAutoRemoval(World); });
diff --git a/Engine/Plugins/WorldMetrics/Source/WorldMetricsTest/Public/WorldMetricsTestUtil.h b/Engine/Plugins/WorldMetrics/Source/WorldMetricsTest/Public/WorldMetricsTestUtil.h
--- a/Engine/Plugins/WorldMetrics/Source/WorldMetricsTest/Public/WorldMetricsTestUtil.h
+++ b/Engine/Plugins/WorldMetrics/Source/WorldMetricsTest/Public/WorldMetricsTestUtil.h
@@ -5,13 +5,16 @@
 #include "Containers/Array.h"
 
 #ifndef REQUIRE_MESSAGE
+
+#include "WorldMetricsLog.h"
+
 #define REQUIRE_MESSAGE(Message, Expr)           \
 	if (!(Expr))                                 \
 	{                                            \
 		UE_LOG(LogWorldMetrics, Error, Message); \
 		return;                                  \
 	}
-#endif
+#endif // REQUIRE_MESSAGE
 
 namespace UE::WorldMetrics
 {
