WorldMetrics - Remove game world dependency from enabling the subsystem.

#rnx

diff --git a/Engine/Plugins/WorldMetrics/Source/WorldMetricsCore/Private/WorldMetricsSubsystem.cpp b/Engine/Plugins/WorldMetrics/Source/WorldMetricsCore/Private/WorldMetricsSubsystem.cpp
--- a/Engine/Plugins/WorldMetrics/Source/WorldMetricsCore/Private/WorldMetricsSubsystem.cpp
+++ b/Engine/Plugins/WorldMetrics/Source/WorldMetricsCore/Private/WorldMetricsSubsystem.cpp
@@ -106,7 +106,7 @@ void UWorldMetricsSubsystem::Enable(bool bEnable)
 			return;
 		}
 
-		if (!GetWorld()->IsGameWorld())
+		if (!CanHaveWorldMetrics(GetWorld()))
 		{
 			return;
 		}
