WorldMetrics API update: updated World Metrics subsystem contract check for method parameters to report a Warning instead of an Error. This change ensures methods can be unit tested. Errors are now reserved for incorrect internal state checks and object creation failures.

#test passed WorldMetrics.Core unit tests

diff --git a/Engine/Plugins/WorldMetrics/Source/WorldMetricsCore/Private/WorldMetricsSubsystem.cpp b/Engine/Plugins/WorldMetrics/Source/WorldMetricsCore/Private/WorldMetricsSubsystem.cpp
--- a/Engine/Plugins/WorldMetrics/Source/WorldMetricsCore/Private/WorldMetricsSubsystem.cpp
+++ b/Engine/Plugins/WorldMetrics/Source/WorldMetricsCore/Private/WorldMetricsSubsystem.cpp
@@ -205,14 +205,14 @@ UWorldMetricInterface* UWorldMetricsSubsystem::CreateMetric(const TSubclassOf<UW
 {
 	if (UNLIKELY(!InMetricClass))
 	{
-		UE_LOG(LogWorldMetrics, Error, TEXT("[%hs] Unexpected null metric class"), __FUNCTION__);
+		UE_LOG(LogWorldMetrics, Warning, TEXT("[%hs] Unexpected null metric class"), __FUNCTION__);
 		return nullptr;
 	}
 
 	if (InMetricClass->HasAnyClassFlags(CLASS_Abstract))
 	{
 		UE_LOG(
-			LogWorldMetrics, Error, TEXT("[%hs] Parameter metric class is abstract: %s"), __FUNCTION__,
+			LogWorldMetrics, Warning, TEXT("[%hs] Parameter metric class is abstract: %s"), __FUNCTION__,
 			*InMetricClass->GetFName().ToString());
 		return nullptr;
 	}
@@ -231,7 +231,7 @@ bool UWorldMetricsSubsystem::ContainsMetric(UWorldMetricInterface* InMetric) con
 {
 	if (UNLIKELY(!InMetric))
 	{
-		UE_LOG(LogWorldMetrics, Error, TEXT("[%hs] Unexpected null metric instance"), __FUNCTION__);
+		UE_LOG(LogWorldMetrics, Warning, TEXT("[%hs] Unexpected null metric instance"), __FUNCTION__);
 		return false;
 	}
 	return Algo::IndexOf(Metrics, InMetric) != INDEX_NONE;
@@ -248,7 +248,7 @@ bool UWorldMetricsSubsystem::AddMetric(UWorldMetricInterface* InMetric)
 {
 	if (UNLIKELY(!InMetric))
 	{
-		UE_LOG(LogWorldMetrics, Error, TEXT("[%hs] Unexpected null metric instance"), __FUNCTION__);
+		UE_LOG(LogWorldMetrics, Warning, TEXT("[%hs] Unexpected null metric instance"), __FUNCTION__);
 		return false;
 	}
 
@@ -279,7 +279,7 @@ bool UWorldMetricsSubsystem::RemoveMetric(UWorldMetricInterface* InMetric)
 {
 	if (UNLIKELY(!InMetric))
 	{
-		UE_LOG(LogWorldMetrics, Error, TEXT("[%hs] Unexpected null metric instance"), __FUNCTION__);
+		UE_LOG(LogWorldMetrics, Warning, TEXT("[%hs] Unexpected null metric instance"), __FUNCTION__);
 		return false;
 	}
 
@@ -394,20 +394,20 @@ UWorldMetricsExtension* UWorldMetricsSubsystem::AcquireExtensionInternal(
 {
 	if (UNLIKELY(!InOwner))
 	{
-		UE_LOG(LogWorldMetrics, Error, TEXT("[%hs] Unexpected invalid owner"), __FUNCTION__);
+		UE_LOG(LogWorldMetrics, Warning, TEXT("[%hs] Unexpected invalid owner"), __FUNCTION__);
 		return nullptr;
 	}
 
 	if (UNLIKELY(!InExtensionClass))
 	{
-		UE_LOG(LogWorldMetrics, Error, TEXT("[%hs] Unexpected null extension class"), __FUNCTION__);
+		UE_LOG(LogWorldMetrics, Warning, TEXT("[%hs] Unexpected null extension class"), __FUNCTION__);
 		return nullptr;
 	}
 
 	if (InExtensionClass->HasAnyClassFlags(CLASS_Abstract))
 	{
 		UE_LOG(
-			LogWorldMetrics, Error, TEXT("[%hs] Parameter extension class is abstract: %s"), __FUNCTION__,
+			LogWorldMetrics, Warning, TEXT("[%hs] Parameter extension class is abstract: %s"), __FUNCTION__,
 			*InExtensionClass->GetFName().ToString());
 		return nullptr;
 	}
@@ -437,7 +437,7 @@ UWorldMetricsExtension* UWorldMetricsSubsystem::AcquireExistingExtension(
 
 	if (!ensure(IndexedOwners.IsValidIndex(ExtensionIndex)))
 	{
-		UE_LOG(LogWorldMetrics, Warning, TEXT("[%hs] Unexpected invalid extension's owner list"), __FUNCTION__);
+		UE_LOG(LogWorldMetrics, Error, TEXT("[%hs] Unexpected invalid extension's owner list"), __FUNCTION__);
 		return nullptr;
 	}
 	IndexedOwners[ExtensionIndex].Emplace(InOwner);
@@ -451,8 +451,8 @@ UWorldMetricsExtension* UWorldMetricsSubsystem::AddExtension(
 {
 	check(Extensions.Num() == IndexedOwners.Num());
 	UWorldMetricsExtension* Extension =
-		NewObject<UWorldMetricsExtension>(this, InExtensionClass, NAME_None, RF_Transient);
-	if (Extension == nullptr)
+		NewObject<UWorldMetricsExtension>(this, InExtensionClass, NAME_None, RF_Transient);	
+	if (UNLIKELY(!Extension))
 	{
 		UE_LOG(
 			LogWorldMetrics, Error, TEXT("[%hs] Failed to create extension of class: %s"), __FUNCTION__,
@@ -482,7 +482,7 @@ bool UWorldMetricsSubsystem::ReleaseExtensionInternal(
 
 	if (UNLIKELY(!InExtensionClass))
 	{
-		UE_LOG(LogWorldMetrics, Error, TEXT("[%hs] Unexpected null extension class"), __FUNCTION__);
+		UE_LOG(LogWorldMetrics, Warning, TEXT("[%hs] Unexpected null extension class"), __FUNCTION__);
 		return false;
 	}
 
diff --git a/Engine/Plugins/WorldMetrics/Source/WorldMetricsTest/Private/WorldMetricsTest.cpp b/Engine/Plugins/WorldMetrics/Source/WorldMetricsTest/Private/WorldMetricsTest.cpp
--- a/Engine/Plugins/WorldMetrics/Source/WorldMetricsTest/Private/WorldMetricsTest.cpp
+++ b/Engine/Plugins/WorldMetrics/Source/WorldMetricsTest/Private/WorldMetricsTest.cpp
@@ -7,6 +7,7 @@
 #endif	// WITH_TESTS
 
 #include "HAL/IConsoleManager.h"
+#include "WorldMetricsLog.h"
 #include "WorldMetricsSubsystem.h"
 #include "WorldMetricsTestTypes.h"
 #include "WorldMetricsTestUtil.h"
@@ -22,8 +23,25 @@ void TestAll(UWorld* World);
 static FAutoConsoleCommandWithWorldAndArgs CmdWorldMetricsSelfTest(
 	TEXT("WorldMetrics.SelfTest"),
 	TEXT("Toggles the World Metrics Subsystem self-test."),
-	FConsoleCommandWithWorldAndArgsDelegate::CreateLambda([](const TArray<FString>& Args, UWorld* World)
-														  { TestAll(World); }),
+	FConsoleCommandWithWorldAndArgsDelegate::CreateLambda(
+		[](const TArray<FString>& Args, UWorld* World)
+		{
+			const UWorldMetricsSubsystem* Subsystem = UWorldMetricsSubsystem::Get(World);
+			if (!ensure(Subsystem))
+			{
+				UE_LOG(LogWorldMetrics, Warning, TEXT("Cannot run SelfTest without a valid World Metrics subsystem."))
+				return;
+			}
+
+			if (Subsystem->HasAnyMetric())
+			{
+				UE_LOG(
+					LogWorldMetrics, Warning,
+					TEXT("Cannot run SelfTest in a World Metrics subsystem with pre-existing metrics."))
+				return;
+			}
+			TestAll(World);
+		}),
 	ECVF_Default);
 
 static const TSubclassOf<UMockWorldMetricBase> MockMetricClasses[] = {
@@ -544,9 +562,6 @@ void TestMultipleMetricIteration(UWorld* World, bool bRandomized)
 
 void TestAll(UWorld* World)
 {
-	UWorldMetricsSubsystem* Subsystem = UWorldMetricsSubsystem::Get(World);
-	REQUIRE_MESSAGE(TEXT("WorldMetricsSubsystem is valid after GetSubsystem."), Subsystem);
-
 	TestZeroState(World);
 
 	// Metrics
@@ -601,7 +616,7 @@ static void ScopedWorldTest(EWorldType::Type WorldType, TFunctionRef<void(UWorld
 }
 }  // namespace Private
 
-TEST_CASE_NAMED(WorldMetricsTestZeroState, "WorldMetrics::TestZeroState", "[WorldMetrics][EngineFilter]")
+TEST_CASE_NAMED(WorldMetricsTestZeroState, "WorldMetrics::TestZeroState", "[WorldMetrics][Core][EngineFilter]")
 {
 	Private::ScopedWorldTest(EWorldType::Editor, [this](UWorld* World) { Private::TestZeroState(World); });
 }
@@ -609,7 +624,7 @@ TEST_CASE_NAMED(WorldMetricsTestZeroState, "WorldMetrics::TestZeroState", "[Worl
 TEST_CASE_NAMED(
 	WorldMetricsTestSingleMetricAddRemove,
 	"WorldMetrics::TestSingleMetricAddRemove",
-	"[WorldMetrics][EngineFilter]")
+	"[WorldMetrics][Core][EngineFilter]")
 {
 	Private::ScopedWorldTest(EWorldType::Editor, [this](UWorld* World) { Private::TestSingleMetricAddRemove(World); });
 }
@@ -617,7 +632,7 @@ TEST_CASE_NAMED(
 TEST_CASE_NAMED(
 	WorldMetricsTestMultipleMetricsAddRemove,
 	"WorldMetrics::TestMultipleMetricsAddRemove",
-	"[WorldMetrics][EngineFilter]")
+	"[WorldMetrics][Core][EngineFilter]")
 {
 	Private::ScopedWorldTest(
 		EWorldType::Editor,
@@ -631,7 +646,7 @@ TEST_CASE_NAMED(
 TEST_CASE_NAMED(
 	WorldMetricsTestMultipleMetricIteration,
 	"WorldMetrics::TestMultipleMetricIteration",
-	"[WorldMetrics][EngineFilter]")
+	"[WorldMetrics][Core][EngineFilter]")
 {
 	Private::ScopedWorldTest(
 		EWorldType::Editor,
@@ -645,7 +660,7 @@ TEST_CASE_NAMED(
 TEST_CASE_NAMED(
 	WorldMetricsTestSingleMetricSingleExtensionAcquireRelease,
 	"WorldMetrics::TestSingleMetricSingleExtensionAcquireRelease",
-	"[WorldMetrics][EngineFilter]")
+	"[WorldMetrics][Core][EngineFilter]")
 {
 	Private::ScopedWorldTest(
 		EWorldType::Editor, [this](UWorld* World) { Private::TestSingleMetricSingleExtensionAcquireRelease(World); });
@@ -654,7 +669,7 @@ TEST_CASE_NAMED(
 TEST_CASE_NAMED(
 	WorldMetricsTestMultipleMetricsSingleExtensionAcquireRelease,
 	"WorldMetrics::TestMultipleMetricsSingleExtensionAcquireRelease",
-	"[WorldMetrics][EngineFilter]")
+	"[WorldMetrics][Core][EngineFilter]")
 {
 	Private::ScopedWorldTest(
 		EWorldType::Editor,
@@ -668,7 +683,7 @@ TEST_CASE_NAMED(
 TEST_CASE_NAMED(
 	WorldMetricsTestMultipleMetricsMultipleExtensionAcquireRelease,
 	"WorldMetrics::TestMultipleMetricsMultipleExtensionAcquireRelease",
-	"[WorldMetrics][EngineFilter]")
+	"[WorldMetrics][Core][EngineFilter]")
 {
 	Private::ScopedWorldTest(
 		EWorldType::Editor,
@@ -682,7 +697,7 @@ TEST_CASE_NAMED(
 TEST_CASE_NAMED(
 	WorldMetricsTestSingleMetricSingleExtensionAutoReleaseOnRemoval,
 	"WorldMetrics::TestSingleMetricSingleExtensionAutoReleaseOnRemoval",
-	"[WorldMetrics][EngineFilter]")
+	"[WorldMetrics][Core][EngineFilter]")
 {
 	Private::ScopedWorldTest(
 		EWorldType::Editor,
@@ -692,7 +707,7 @@ TEST_CASE_NAMED(
 TEST_CASE_NAMED(
 	WorldMetricsTestMultipleMetricsMultipleExtensionAutoReleaseOnRemoval,
 	"WorldMetrics::TestMultipleMetricsMultipleExtensionAutoReleaseOnRemoval",
-	"[WorldMetrics][EngineFilter]")
+	"[WorldMetrics][Core][EngineFilter]")
 {
 	Private::ScopedWorldTest(
 		EWorldType::Editor,
@@ -706,7 +721,7 @@ TEST_CASE_NAMED(
 TEST_CASE_NAMED(
 	WorldMetricsTestSingleMetricOrphanExtensionAutoRemoval,
 	"WorldMetrics::TestSingleMetricOrphanExtensionAutoRemoval",
-	"[WorldMetrics][EngineFilter]")
+	"[WorldMetrics][Core][EngineFilter]")
 {
 	Private::ScopedWorldTest(
 		EWorldType::Editor, [this](UWorld* World) { Private::TestSingleMetricOrphanExtensionAutoRemoval(World); });
