[World Metrics] Convert extensions to be stored in a map instead of an array - addresses a case where removing an extension causes re-entrancy in the removal function, invalidating extension indices. Added test case which results in a crash without this fix.

#tests Ran unit tests, including the new test case - no crash encountered.

diff --git a/Engine/Plugins/WorldMetrics/Source/WorldMetricsCore/Private/WorldMetricsSubsystem.cpp b/Engine/Plugins/WorldMetrics/Source/WorldMetricsCore/Private/WorldMetricsSubsystem.cpp
--- a/Engine/Plugins/WorldMetrics/Source/WorldMetricsCore/Private/WorldMetricsSubsystem.cpp
+++ b/Engine/Plugins/WorldMetrics/Source/WorldMetricsCore/Private/WorldMetricsSubsystem.cpp
@@ -51,6 +51,17 @@ void UWorldMetricsSubsystem::Deinitialize()
 	Super::Deinitialize();
 }
 
+void UWorldMetricsSubsystem::AddReferencedObjects(UObject* InThis, FReferenceCollector& Collector)
+{
+	UWorldMetricsSubsystem* This = CastChecked<UWorldMetricsSubsystem>(InThis);
+
+	for (TPair<TSubclassOf<UWorldMetricsExtension>, FExtension>& ExtensionPair : This->Extensions)
+	{
+		Collector.AddReferencedObject(ExtensionPair.Key.GetGCPtr());
+		Collector.AddReferencedObject(ExtensionPair.Value.Instance);
+	}
+}
+
 void UWorldMetricsSubsystem::BeginDestroy()
 {
 	if (!ensure(!IsEnabled()))
@@ -92,7 +103,6 @@ void UWorldMetricsSubsystem::Clear()
 	RemoveAllMetrics();
 
 	Extensions.Reset();
-	IndexedOwners.Reset();
 
 	UE_LOG(LogWorldMetrics, Log, TEXT("[%hs]"), __FUNCTION__);
 }
@@ -381,13 +391,6 @@ bool UWorldMetricsSubsystem::ReleaseExtension(
 	return ReleaseExtensionInternal(InExtensionOwner, InExtensionClass);
 }
 
-int32 UWorldMetricsSubsystem::GetExtensionIndex(const TSubclassOf<UWorldMetricsExtension>& InExtensionClass) const
-{
-	return Algo::IndexOfByPredicate(
-		Extensions,
-		[InExtensionClass](UWorldMetricsExtension* Extension) { return Extension->GetClass() == InExtensionClass; });
-}
-
 UWorldMetricsExtension* UWorldMetricsSubsystem::AcquireExtensionInternal(
 	UObject* InOwner,
 	const TSubclassOf<UWorldMetricsExtension>& InExtensionClass)
@@ -429,30 +432,24 @@ UWorldMetricsExtension* UWorldMetricsSubsystem::AcquireExistingExtension(
 	UObject* InOwner,
 	const TSubclassOf<UWorldMetricsExtension>& InExtensionClass)
 {
-	const int32 ExtensionIndex = GetExtensionIndex(InExtensionClass);
-	if (ExtensionIndex == INDEX_NONE)
+	FExtension* Extension = Extensions.Find(InExtensionClass);
+	if (!Extension)
 	{
 		return nullptr;
 	}
 
-	if (!ensure(IndexedOwners.IsValidIndex(ExtensionIndex)))
-	{
-		UE_LOG(LogWorldMetrics, Error, TEXT("[%hs] Unexpected invalid extension's owner list"), __FUNCTION__);
-		return nullptr;
-	}
-	IndexedOwners[ExtensionIndex].Emplace(InOwner);
+	Extension->Owners.Emplace(InOwner);
 
-	return Extensions[ExtensionIndex];
+	return Extension->Instance;
 }
 
 UWorldMetricsExtension* UWorldMetricsSubsystem::AddExtension(
 	UObject* InOwner,
 	const TSubclassOf<UWorldMetricsExtension>& InExtensionClass)
 {
-	check(Extensions.Num() == IndexedOwners.Num());
-	UWorldMetricsExtension* Extension =
+	UWorldMetricsExtension* ExtensionInstance =
 		NewObject<UWorldMetricsExtension>(this, InExtensionClass, NAME_None, RF_Transient);	
-	if (UNLIKELY(!Extension))
+	if (UNLIKELY(!ExtensionInstance))
 	{
 		UE_LOG(
 			LogWorldMetrics, Error, TEXT("[%hs] Failed to create extension of class: %s"), __FUNCTION__,
@@ -460,15 +457,16 @@ UWorldMetricsExtension* UWorldMetricsSubsystem::AddExtension(
 		return nullptr;
 	}
 
-	Extensions.Emplace(Extension);
-	IndexedOwners.Emplace_GetRef().Emplace(InOwner);
-	Extension->Initialize();
+	FExtension& Extension = Extensions.Add(InExtensionClass);
+	Extension.Instance = ExtensionInstance;
+	Extension.Owners.Emplace(InOwner);
+	ExtensionInstance->Initialize();
 
 	UE_LOG(
 		LogWorldMetrics, Log, TEXT("[%hs] Added extension of class: %s"), __FUNCTION__,
 		*InExtensionClass->GetFName().ToString());
 
-	return Extension;
+	return ExtensionInstance;
 }
 
 bool UWorldMetricsSubsystem::ReleaseExtensionInternal(
@@ -486,38 +484,38 @@ bool UWorldMetricsSubsystem::ReleaseExtensionInternal(
 		return false;
 	}
 
-	const int32 ExtensionIndex = GetExtensionIndex(InExtensionClass);
-	if (ExtensionIndex == INDEX_NONE)
+	FExtension* Extension = Extensions.Find(InExtensionClass);
+	if (!Extension)
 	{
 		return false;
 	}
 
-	if (UNLIKELY(!IndexedOwners[ExtensionIndex].Remove(InOwner)))
+	Extension->Instance->OnRelease(InOwner);
+	if (UNLIKELY(!Extension->Owners.Remove(InOwner)))
 	{
 		UE_LOG(
 			LogWorldMetrics, Error, TEXT("[%hs] Parameter object doesn't own this extension: %s"), __FUNCTION__,
 			*InExtensionClass->GetFName().ToString());
 	}
-	Extensions[ExtensionIndex]->OnRelease(InOwner);
 
-	TryRemoveExtensionAt(ExtensionIndex);
+	TryRemoveExtension(InExtensionClass);
 
 	return true;
 }
 
-bool UWorldMetricsSubsystem::TryRemoveExtensionAt(int32 ExtensionIndex)
+bool UWorldMetricsSubsystem::TryRemoveExtension(const TSubclassOf<UWorldMetricsExtension>& InExtensionClass)
 {
-	if (IndexedOwners[ExtensionIndex].IsEmpty())
+	FExtension* Extension = Extensions.Find(InExtensionClass);
+	if (Extension && Extension->Owners.IsEmpty())
 	{
-		UWorldMetricsExtension* Extension = Extensions[ExtensionIndex];
-		Extension->Deinitialize();
-		Extension->MarkAsGarbage();
+		UWorldMetricsExtension* ExtensionInstance = Extension->Instance;
+		ExtensionInstance->Deinitialize();
+		ExtensionInstance->MarkAsGarbage();
 
-		Extensions.RemoveAtSwap(ExtensionIndex);
-		IndexedOwners.RemoveAtSwap(ExtensionIndex);
+		Extensions.Remove(InExtensionClass);
 		UE_LOG(
 			LogWorldMetrics, Log, TEXT("[%hs] Removed extension of class %s"), __FUNCTION__,
-			*Extension->GetClass()->GetFName().ToString());
+			*ExtensionInstance->GetClass()->GetFName().ToString());
 
 		return true;
 	}
@@ -526,50 +524,51 @@ bool UWorldMetricsSubsystem::TryRemoveExtensionAt(int32 ExtensionIndex)
 
 void UWorldMetricsSubsystem::VerifyMetricReleasedAllExtensions(UWorldMetricInterface* InMetric)
 {
-	check(Extensions.Num() == IndexedOwners.Num());
-
-	TArray<int32, TInlineAllocator<DefaultExtensionCapacity>> StaleExtensionIndices;
-	for (int32 ExtensionIndex = 0; ExtensionIndex < Extensions.Num(); ++ExtensionIndex)
+	TArray<TSubclassOf<UWorldMetricsExtension>, TInlineAllocator<DefaultExtensionCapacity>> StaleExtensionClasses;
+	for (TPair<TSubclassOf<UWorldMetricsExtension>, FExtension>& ExtensionPair : Extensions)
 	{
-		if (UNLIKELY(IndexedOwners[ExtensionIndex].Remove(InMetric)))
+		if (UNLIKELY(ExtensionPair.Value.Owners.Remove(InMetric)))
 		{
-			StaleExtensionIndices.Emplace(ExtensionIndex);
-			UWorldMetricsExtension* Extension = Extensions[ExtensionIndex];
+			StaleExtensionClasses.Emplace(ExtensionPair.Key);
+			UWorldMetricsExtension* Extension = ExtensionPair.Value.Instance;
 			UE_LOG(
 				LogWorldMetrics, Warning, TEXT("[%hs] World metric %s did not release extension %s in Deinitialize"),
 				__FUNCTION__, *InMetric->GetFName().ToString(), *Extension->GetName());
 		}
 	}
 
-	// Reverse iterate stale extension indices in case the extension can be removed
-	for (int32 i = StaleExtensionIndices.Num(); i-- > 0;)
+	for (const TSubclassOf<UWorldMetricsExtension>& ExtensionClass : StaleExtensionClasses)
 	{
-		TryRemoveExtensionAt(StaleExtensionIndices[i]);
+		TryRemoveExtension(ExtensionClass);
 	}
 }
 
 void UWorldMetricsSubsystem::VerifyRemoveOrphanExtensions()
 {
-	check(Extensions.Num() == IndexedOwners.Num());
-
 	if (Metrics.IsEmpty())
 	{
 		if (UNLIKELY(!Extensions.IsEmpty()))
 		{
-			for (UWorldMetricsExtension* Extension : Extensions)
+			TArray<TSubclassOf<UWorldMetricsExtension>, TInlineAllocator<DefaultExtensionCapacity>> ExtensionsToRemove;
+			Extensions.GenerateKeyArray(ExtensionsToRemove);
+			for (const TSubclassOf<UWorldMetricsExtension>& ExtensionClass : ExtensionsToRemove)
 			{
-				check(Extension);
-
-				UE_LOG(
-					LogWorldMetrics, Warning,
-					TEXT("[%hs] World extension %s was acquired by another extension but never released."),
-					__FUNCTION__, *Extension->GetName());
-
-				Extension->Deinitialize();
-				Extension->MarkAsGarbage();
+				if (const FExtension* Extension = Extensions.Find(ExtensionClass))
+				{
+					UWorldMetricsExtension* ExtensionInstance = Extension->Instance;
+					check(ExtensionInstance);
+
+					UE_LOG(
+						LogWorldMetrics, Warning,
+						TEXT("[%hs] World extension %s was acquired by another extension but never released."),
+						__FUNCTION__, *ExtensionInstance->GetName());
+
+					ExtensionInstance->Deinitialize();
+					ExtensionInstance->MarkAsGarbage();
+					Extensions.Remove(ExtensionClass);
+				}
 			}
 			Extensions.Reset();
-			IndexedOwners.Reset();
 		}
 	}
 }
diff --git a/Engine/Plugins/WorldMetrics/Source/WorldMetricsCore/Public/WorldMetricsSubsystem.h b/Engine/Plugins/WorldMetrics/Source/WorldMetricsCore/Public/WorldMetricsSubsystem.h
--- a/Engine/Plugins/WorldMetrics/Source/WorldMetricsCore/Public/WorldMetricsSubsystem.h
+++ b/Engine/Plugins/WorldMetrics/Source/WorldMetricsCore/Public/WorldMetricsSubsystem.h
@@ -41,6 +41,7 @@ public:
 	//~ End USubsystem
 
 	//~ Begin UObject
+	WORLDMETRICSCORE_API static void AddReferencedObjects(UObject* InThis, FReferenceCollector& Collector);
 	WORLDMETRICSCORE_API virtual void BeginDestroy() override;
 	//~ End UObject
 
@@ -222,11 +223,16 @@ private:
 	UPROPERTY(Transient)
 	TArray<TObjectPtr<UWorldMetricInterface>> Metrics;
 
-	UPROPERTY(Transient)
-	TArray<TObjectPtr<UWorldMetricsExtension>> Extensions;
-
 	using FOwnerList = TSet<UObject*, DefaultKeyFuncs<UObject*>, TInlineSetAllocator<DefaultOwnerListCapacity>>;
-	TArray<FOwnerList, TInlineAllocator<DefaultExtensionCapacity>> IndexedOwners;
+
+	struct FExtension
+	{
+		TObjectPtr<UWorldMetricsExtension> Instance;
+		FOwnerList Owners;
+	};
+
+	// Exposed to GC via AddReferencedObjects.
+	TMap<TSubclassOf<UWorldMetricsExtension>, FExtension, TInlineSetAllocator<DefaultExtensionCapacity>> Extensions;
 
 public:
 	UPROPERTY(Config)
@@ -264,14 +270,6 @@ private:
 	 */
 	void Clear();
 
-	/**
-	 * Returns the live extension list index corresponding to the parameter extension class.
-	 *
-	 * @param InExtensionClass: the class type of the extension.
-	 * @return a valid index value or INDEX_NONE if not matching extension exists.
-	 */
-	int32 GetExtensionIndex(const TSubclassOf<UWorldMetricsExtension>& InExtensionClass) const;
-
 	UWorldMetricsExtension* AcquireExtensionInternal(
 		UObject* InOwner,
 		const TSubclassOf<UWorldMetricsExtension>& InExtensionClass);
@@ -284,7 +282,7 @@ private:
 
 	bool ReleaseExtensionInternal(UObject* InOwner, const TSubclassOf<UWorldMetricsExtension>& InExtensionClass);
 
-	bool TryRemoveExtensionAt(int32 ExtensionIndex);
+	bool TryRemoveExtension(const TSubclassOf<UWorldMetricsExtension>& InExtensionClass);
 
 	/**
 	 * Verifies the parameter metric has released any acquired extensions, releasing any pending ones. It's the metric's
diff --git a/Engine/Plugins/WorldMetrics/Source/WorldMetricsTest/Private/WorldMetricsTest.cpp b/Engine/Plugins/WorldMetrics/Source/WorldMetricsTest/Private/WorldMetricsTest.cpp
--- a/Engine/Plugins/WorldMetrics/Source/WorldMetricsTest/Private/WorldMetricsTest.cpp
+++ b/Engine/Plugins/WorldMetrics/Source/WorldMetricsTest/Private/WorldMetricsTest.cpp
@@ -333,6 +333,96 @@ void TestSingleMetricSingleExtensionAcquireRelease(UWorld* World)
 	TestZeroState(World);
 }
 
+void TestMultipleMetricsSingleExtensionAcquireReleaseChildExtension(UWorld* World)
+{
+	TestZeroState(World);
+
+	UWorldMetricsSubsystem* Subsystem = UWorldMetricsSubsystem::Get(World);
+
+	UMockWorldMetricA* MetricA = Subsystem->AddMetric<UMockWorldMetricA>();
+	REQUIRE_MESSAGE(TEXT("ContainsMetric failed"), Subsystem->ContainsMetric(MetricA));
+
+	UMockWorldMetricB* MetricB = Subsystem->AddMetric<UMockWorldMetricB>();
+	REQUIRE_MESSAGE(TEXT("ContainsMetric failed"), Subsystem->ContainsMetric(MetricB));
+
+	REQUIRE_MESSAGE(
+		TEXT("WorldMetricsSubsystem should be enabled by default if there are metrics."), Subsystem->IsEnabled());
+
+	const UMockWorldMetricsExtensionBase* ExtensionA = nullptr;
+	const UMockWorldMetricsExtensionBase* ExtensionB = nullptr;
+	const UMockWorldMetricsExtensionBase* ExtensionC = nullptr;
+
+	// Create lambda to handle acquiring/releasing other extensions when an extension is initialized/deinitialized.
+	auto ExtensionInitializeDeinitializeHandler = [&ExtensionC](UWorldMetricsExtension* Extension, bool bInitialize)
+	{
+		// B acquires C.
+		if (Extension->IsA<UMockWorldMetricsExtensionB>())
+		{
+			if (bInitialize)
+			{
+				ExtensionC = Extension->GetOwner().AcquireExtension<UMockWorldMetricsExtensionC>(Extension);
+			}
+			else
+			{
+				Extension->GetOwner().ReleaseExtension<UMockWorldMetricsExtensionC>(Extension);
+			}
+		}
+	};
+
+	// TGuardValue to make sure the static OnInitializeDeinitialize delegate is unbound on scope exit.
+	TGuardValue InitializeDeinitializeScope(
+		UMockWorldMetricsExtensionBase::OnInitializeDeinitialize,
+		UMockWorldMetricsExtensionBase::FInitializeDeinitializeDelegate::CreateLambda(ExtensionInitializeDeinitializeHandler));
+
+	// Acquire Extension A for Metric A.
+	{
+		ExtensionA = MetricA->GetOwner().AcquireExtension<UMockWorldMetricsExtensionA>(MetricA);
+		REQUIRE_MESSAGE(TEXT("Acquire Test Extension A failed."), ExtensionA);
+		REQUIRE_MESSAGE(TEXT("Test Extension A initialization failed"), ExtensionA->InitializeCount == 1);
+
+		REQUIRE_MESSAGE(TEXT("WorldMetricsSubsystem should have extensions."), Subsystem->HasAnyExtension());
+		REQUIRE_MESSAGE(TEXT("WorldMetricsSubsystem should have 1 extension."), Subsystem->NumExtensions() == 1);
+	}
+
+	// Acquire Extension B (and C, transitively) for Metric B.
+	{
+		ExtensionB = MetricB->GetOwner().AcquireExtension<UMockWorldMetricsExtensionB>(MetricB);
+		REQUIRE_MESSAGE(TEXT("Acquire Test Extension B failed."), ExtensionB);
+		REQUIRE_MESSAGE(TEXT("Test Extension B initialization failed"), ExtensionB->InitializeCount == 1);
+
+		REQUIRE_MESSAGE(TEXT("Acquire Test Extension C failed."), ExtensionC);
+		REQUIRE_MESSAGE(TEXT("Test Extension C initialization failed"), ExtensionC->InitializeCount == 1);
+
+		REQUIRE_MESSAGE(TEXT("WorldMetricsSubsystem should have extensions."), Subsystem->HasAnyExtension());
+		REQUIRE_MESSAGE(TEXT("WorldMetricsSubsystem should have 3 extensions."), Subsystem->NumExtensions() == 3);
+	}
+
+	// Release extension for Metric A.
+	{
+		REQUIRE_MESSAGE(
+			TEXT("Test Extension release failed."),
+			MetricA->GetOwner().ReleaseExtension<UMockWorldMetricsExtensionA>(MetricA));
+		REQUIRE_MESSAGE(
+			TEXT("Test Extension deinitialization failed after release."), ExtensionA->DeinitializeCount == 1);
+	}
+
+	// Release extension for Metric B.
+	{
+		REQUIRE_MESSAGE(
+			TEXT("Test Extension B release failed."),
+			MetricB->GetOwner().ReleaseExtension<UMockWorldMetricsExtensionB>(MetricB));
+		REQUIRE_MESSAGE(
+			TEXT("Test Extension B deinitialization failed after release."), ExtensionB->DeinitializeCount == 1);
+		REQUIRE_MESSAGE(
+			TEXT("Test Extension C deinitialization failed after release."), ExtensionC->DeinitializeCount == 1);
+	}
+
+	Subsystem->RemoveMetric(MetricB);
+	Subsystem->RemoveMetric(MetricA);
+
+	TestZeroState(World);
+}
+
 void TestMultipleMetricsSingleExtensionAcquireRelease(UWorld* World, bool bRandomized)
 {
 	TestZeroState(World);
@@ -575,6 +665,7 @@ void TestAll(UWorld* World)
 	TestSingleMetricSingleExtensionAcquireRelease(World);
 	TestMultipleMetricsSingleExtensionAcquireRelease(World, false);
 	TestMultipleMetricsSingleExtensionAcquireRelease(World, true);
+	TestMultipleMetricsSingleExtensionAcquireReleaseChildExtension(World);
 	TestMultipleMetricsMultipleExtensionAcquireRelease(World, false);
 	TestMultipleMetricsMultipleExtensionAcquireRelease(World, true);
 
@@ -681,6 +772,19 @@ TEST_CASE_NAMED(
 		});
 }
 
+TEST_CASE_NAMED(
+	WorldMetricsTestMultipleMetricsSingleExtensionAcquireReleaseChildExtension,
+	"WorldMetrics::TestMultipleMetricsSingleExtensionAcquireReleaseChildExtension",
+	"[WorldMetrics][Core][EngineFilter]")
+{
+	Private::ScopedWorldTest(
+		EWorldType::Editor,
+		[this](UWorld* World)
+		{
+			Private::TestMultipleMetricsSingleExtensionAcquireReleaseChildExtension(World);
+		});
+}
+
 TEST_CASE_NAMED(
 	WorldMetricsTestMultipleMetricsMultipleExtensionAcquireRelease,
 	"WorldMetrics::TestMultipleMetricsMultipleExtensionAcquireRelease",
diff --git a/Engine/Plugins/WorldMetrics/Source/WorldMetricsTest/Private/WorldMetricsTestTypes.cpp b/Engine/Plugins/WorldMetrics/Source/WorldMetricsTest/Private/WorldMetricsTestTypes.cpp
--- a/Engine/Plugins/WorldMetrics/Source/WorldMetricsTest/Private/WorldMetricsTestTypes.cpp
+++ b/Engine/Plugins/WorldMetrics/Source/WorldMetricsTest/Private/WorldMetricsTestTypes.cpp
@@ -27,14 +27,18 @@ void UMockWorldMetricBase::Update(float /*DeltaTimeInSeconds*/)
 // UMockWorldMetricsExtensionBase
 //---------------------------------------------------------------------------------------------------------------------
 
+UMockWorldMetricsExtensionBase::FInitializeDeinitializeDelegate UMockWorldMetricsExtensionBase::OnInitializeDeinitialize;
+
 void UMockWorldMetricsExtensionBase::Initialize()
 {
 	++InitializeCount;
+	OnInitializeDeinitialize.ExecuteIfBound(this, true);
 }
 
 void UMockWorldMetricsExtensionBase::Deinitialize()
 {
 	++DeinitializeCount;
+	OnInitializeDeinitialize.ExecuteIfBound(this, false);
 }
 
 void UMockWorldMetricsExtensionBase::OnAcquire(UObject* InOwner)
diff --git a/Engine/Plugins/WorldMetrics/Source/WorldMetricsTest/Public/WorldMetricsTestTypes.h b/Engine/Plugins/WorldMetrics/Source/WorldMetricsTest/Public/WorldMetricsTestTypes.h
--- a/Engine/Plugins/WorldMetrics/Source/WorldMetricsTest/Public/WorldMetricsTestTypes.h
+++ b/Engine/Plugins/WorldMetrics/Source/WorldMetricsTest/Public/WorldMetricsTestTypes.h
@@ -44,6 +44,10 @@ public:
 	int32 OnAcquireCount = 0;
 	int32 OnReleaseCount = 0;
 
+	// Static delegate to allow injecting logic into Initialize/Deinitialize without needing to create new classes.
+	using FInitializeDeinitializeDelegate = TDelegate<void(UWorldMetricsExtension* Extension, bool bInitialize)>;
+	static FInitializeDeinitializeDelegate OnInitializeDeinitialize;
+
 	virtual SIZE_T GetAllocatedSize() const override
 	{
 		return 0;
