[WorldMetrics API] Updated subsystem's metric/extension logging to Verbose from Log.

diff --git a/Engine/Plugins/WorldMetrics/Source/WorldMetricsCore/Private/WorldMetricsSubsystem.cpp b/Engine/Plugins/WorldMetrics/Source/WorldMetricsCore/Private/WorldMetricsSubsystem.cpp
--- a/Engine/Plugins/WorldMetrics/Source/WorldMetricsCore/Private/WorldMetricsSubsystem.cpp
+++ b/Engine/Plugins/WorldMetrics/Source/WorldMetricsCore/Private/WorldMetricsSubsystem.cpp
@@ -37,14 +37,14 @@ bool UWorldMetricsSubsystem::ShouldCreateSubsystem(UObject* Outer) const
 
 void UWorldMetricsSubsystem::Initialize(FSubsystemCollectionBase& Collection)
 {
-	UE_LOG(LogWorldMetrics, Log, TEXT("[%hs]"), __FUNCTION__);
+	UE_LOG(LogWorldMetrics, Verbose, TEXT("[%hs]"), __FUNCTION__);
 
 	Super::Initialize(Collection);
 }
 
 void UWorldMetricsSubsystem::Deinitialize()
 {
-	UE_LOG(LogWorldMetrics, Log, TEXT("[%hs]"), __FUNCTION__);
+	UE_LOG(LogWorldMetrics, Verbose, TEXT("[%hs]"), __FUNCTION__);
 
 	Clear();
 
@@ -104,7 +104,7 @@ void UWorldMetricsSubsystem::Clear()
 
 	Extensions.Reset();
 
-	UE_LOG(LogWorldMetrics, Log, TEXT("[%hs]"), __FUNCTION__);
+	UE_LOG(LogWorldMetrics, Verbose, TEXT("[%hs]"), __FUNCTION__);
 }
 
 void UWorldMetricsSubsystem::Enable(bool bEnable)
@@ -279,7 +279,7 @@ bool UWorldMetricsSubsystem::AddMetric(UWorldMetricInterface* InMetric)
 	}
 
 	UE_LOG(
-		LogWorldMetrics, Log, TEXT("[%hs] Added metric of class %s."), __FUNCTION__,
+		LogWorldMetrics, Verbose, TEXT("[%hs] Added metric of class %s."), __FUNCTION__,
 		*InMetric->GetClass()->GetFName().ToString());
 
 	return true;
@@ -312,7 +312,7 @@ bool UWorldMetricsSubsystem::RemoveMetric(UWorldMetricInterface* InMetric)
 	}
 
 	UE_LOG(
-		LogWorldMetrics, Log, TEXT("[%hs] Removed metric of class %s"), __FUNCTION__,
+		LogWorldMetrics, Verbose, TEXT("[%hs] Removed metric of class %s"), __FUNCTION__,
 		*InMetric->GetClass()->GetFName().ToString());
 
 	return true;
@@ -463,7 +463,7 @@ UWorldMetricsExtension* UWorldMetricsSubsystem::AddExtension(
 	ExtensionInstance->Initialize();
 
 	UE_LOG(
-		LogWorldMetrics, Log, TEXT("[%hs] Added extension of class: %s"), __FUNCTION__,
+		LogWorldMetrics, Verbose, TEXT("[%hs] Added extension of class: %s"), __FUNCTION__,
 		*InExtensionClass->GetFName().ToString());
 
 	return ExtensionInstance;
@@ -514,7 +514,7 @@ bool UWorldMetricsSubsystem::TryRemoveExtension(const TSubclassOf<UWorldMetricsE
 
 		Extensions.Remove(InExtensionClass);
 		UE_LOG(
-			LogWorldMetrics, Log, TEXT("[%hs] Removed extension of class %s"), __FUNCTION__,
+			LogWorldMetrics, Verbose, TEXT("[%hs] Removed extension of class %s"), __FUNCTION__,
 			*ExtensionInstance->GetClass()->GetFName().ToString());
 
 		return true;
