[World Metrics API] Added a GetOrAdd method to WorldMetricCollection.
* Updated TestMetricCollection unit test.

#tests passed WorldMetrics Core unit tests (12 x 9 iterations).

diff --git a/Engine/Plugins/WorldMetrics/Source/WorldMetricsCore/Private/WorldMetricCollection.cpp b/Engine/Plugins/WorldMetrics/Source/WorldMetricsCore/Private/WorldMetricCollection.cpp
--- a/Engine/Plugins/WorldMetrics/Source/WorldMetricsCore/Private/WorldMetricCollection.cpp
+++ b/Engine/Plugins/WorldMetrics/Source/WorldMetricsCore/Private/WorldMetricCollection.cpp
@@ -156,6 +156,41 @@ bool FWorldMetricCollection::Add(const TSubclassOf<UWorldMetricInterface>& InMet
 	return true;
 }
 
+UWorldMetricInterface* FWorldMetricCollection::GetOrAdd(const TSubclassOf<UWorldMetricInterface>& InMetricClass)
+{
+	UWorldMetricsSubsystem* SubsystemPtr = Subsystem.Get();
+	if (!ensure(ValidateSubsystem(SubsystemPtr, __FUNCTION__)))
+	{
+		return nullptr;
+	}
+
+	UClass* ClassPtr = InMetricClass.Get();
+	if (UNLIKELY(!ValidateClass(ClassPtr, __FUNCTION__)))
+	{
+		return nullptr;
+	}
+
+	const int32 Index = GetMetricIndex(ClassPtr);
+	if (Index != INDEX_NONE)
+	{
+		return Metrics[Index];
+	}
+
+	UWorldMetricInterface* Metric = SubsystemPtr->CreateMetric(ClassPtr);
+	if (UNLIKELY(!Metric))
+	{
+		UE_LOGFMT(LogWorldMetrics, Log, "[{Function}] Unexpected null metric.", __FUNCTION__);
+		return nullptr;
+	}
+
+	Metrics.Emplace(Metric);
+	if (IsEnabled())
+	{
+		SubsystemPtr->AddMetric(Metric);
+	}
+	return Metric;
+}
+
 bool FWorldMetricCollection::Remove(const TSubclassOf<UWorldMetricInterface>& InMetricClass)
 {
 	UWorldMetricsSubsystem* SubsystemPtr = Subsystem.Get();
diff --git a/Engine/Plugins/WorldMetrics/Source/WorldMetricsCore/Public/WorldMetricCollection.h b/Engine/Plugins/WorldMetrics/Source/WorldMetricsCore/Public/WorldMetricCollection.h
--- a/Engine/Plugins/WorldMetrics/Source/WorldMetricsCore/Public/WorldMetricCollection.h
+++ b/Engine/Plugins/WorldMetrics/Source/WorldMetricsCore/Public/WorldMetricCollection.h
@@ -107,17 +107,26 @@ public:
 	 */
 	WORLDMETRICSCORE_API bool Add(const TSubclassOf<UWorldMetricInterface>& InMetricClass);
 
+	template <typename MetricType>
+		requires(std::is_base_of_v<UWorldMetricInterface, MetricType>)
+	bool Add()
+	{
+		return Add(MetricType::StaticClass());
+	}
+
 	/**
-	 * Adds a metric instance of the type parameter class unless a metric of the same class already exists.
+	 * Gets the metric instance of the parameter class or adds a new one if missing.
 	 *
-	 * @tparam MetricType The UWorldMetricInterface's subclass of the metric to add.
-	 * @return True if a metric of the parameter metric class is added and false otherwise.
+	 * @param InMetricClass The class of the metric to add or get.
+	 * @return The added or previously existing instance of the parameter class.
 	 */
+	[[nodiscard]] WORLDMETRICSCORE_API UWorldMetricInterface* GetOrAdd(const TSubclassOf<UWorldMetricInterface>& InMetricClass);
+
 	template <typename MetricType>
 		requires(std::is_base_of_v<UWorldMetricInterface, MetricType>)
-	bool Add()
+	[[nodiscard]] MetricType* GetOrAdd()
 	{
-		return Add(MetricType::StaticClass());
+		return static_cast<MetricType*>(GetOrAdd(MetricType::StaticClass()));
 	}
 
 	/**
diff --git a/Engine/Plugins/WorldMetrics/Source/WorldMetricsTest/Private/WorldMetricsTest.cpp b/Engine/Plugins/WorldMetrics/Source/WorldMetricsTest/Private/WorldMetricsTest.cpp
--- a/Engine/Plugins/WorldMetrics/Source/WorldMetricsTest/Private/WorldMetricsTest.cpp
+++ b/Engine/Plugins/WorldMetrics/Source/WorldMetricsTest/Private/WorldMetricsTest.cpp
@@ -748,6 +748,13 @@ void TestMetricCollection(UWorld* World, bool bRandomized)
 	// Remove
 	REQUIRE_MESSAGE(TEXT("Unexpected metric collection 'Remove<UMockWorldMetricA>' failure."), Collection.Remove<UMockWorldMetricA>());
 
+	// GetOrAdd
+	UMockWorldMetricA* MetricA = Collection.GetOrAdd<UMockWorldMetricA>();
+	REQUIRE_MESSAGE(TEXT("Unexpected metric collection 'GetOrAdd<UMockWorldMetricA>' failure."), MetricA != nullptr);	
+	UMockWorldMetricA* MetricB = Collection.GetOrAdd<UMockWorldMetricA>();
+	REQUIRE_MESSAGE(TEXT("Unexpected metric collection 'GetOrAdd<UMockWorldMetricA>' failure."), MetricB != nullptr);
+	REQUIRE_MESSAGE(TEXT("Unexpected different metric instances."), MetricA == MetricB);
+
 	Collection.Reset();
 
 	TestZeroState(World);
