[Spatial Metrics API] Removed the location metric requirement to sample and collect metrics.
* Added a subscript operator to WorldMetricCollection.
* Added NumValues to USpatialMetricsMetric.
* Updated FSpatialMetricsSampler.
* Updated FSpatialMetricsSampleCollector.
* Updated unit tests.

#tests passed SpatialMetricsCore unit tests (16 x 50).

diff --git a/Engine/Plugins/WorldMetrics/Source/WorldMetricsCore/Public/WorldMetricCollection.h b/Engine/Plugins/WorldMetrics/Source/WorldMetricsCore/Public/WorldMetricCollection.h
--- a/Engine/Plugins/WorldMetrics/Source/WorldMetricsCore/Public/WorldMetricCollection.h
+++ b/Engine/Plugins/WorldMetrics/Source/WorldMetricsCore/Public/WorldMetricCollection.h
@@ -32,6 +32,16 @@ public:
 	WORLDMETRICSCORE_API FWorldMetricCollection() = default;
 	WORLDMETRICSCORE_API ~FWorldMetricCollection();
 
+	[[nodiscard]] const UWorldMetricInterface* operator[](int32 Index) const
+	{
+		return Metrics[Index].Get();
+	}
+
+	[[nodiscard]] UWorldMetricInterface* operator[](int32 Index)
+	{
+		return Metrics[Index].Get();
+	}
+
 	/**
 	 * Initializes the collection. Initializing a container with a valid object resets the current collection.
 	 * @param InOuter: the object that must provide a valid World Metrics Subsystem through its world object.
