[Spatial Metrics Profiler] Refactor to support loading plugins properly, loading in ini files etc. Instead of using a commandlet, just spawn the required tabs after engine init, and ensure that all required plugins are enabled for this program.

#tests Compiled SpatialMetricsProfiler and ran locally.

#changelist validated

diff --git a/Engine/Plugins/Messaging/MessagingDebugger/MessagingDebugger.uplugin b/Engine/Plugins/Messaging/MessagingDebugger/MessagingDebugger.uplugin
--- a/Engine/Plugins/Messaging/MessagingDebugger/MessagingDebugger.uplugin
+++ b/Engine/Plugins/Messaging/MessagingDebugger/MessagingDebugger.uplugin
@@ -10,7 +10,10 @@
 	"Category" : "Messaging",
 	"EnabledByDefault" : false,
 	"IsBetaVersion": true,
-
+	"SupportedPrograms":
+	[
+		"SpatialMetricsProfiler"
+	],
 	"Modules" :
 	[
 		{
diff --git a/Engine/Plugins/Messaging/UdpMessaging/UdpMessaging.uplugin b/Engine/Plugins/Messaging/UdpMessaging/UdpMessaging.uplugin
--- a/Engine/Plugins/Messaging/UdpMessaging/UdpMessaging.uplugin
+++ b/Engine/Plugins/Messaging/UdpMessaging/UdpMessaging.uplugin
@@ -24,7 +24,8 @@
 		"CrashReportClientEditor",
 		"LiveLinkHub",
 		"TimeSynchronizationWriter",
-        "LiveLinkProviderSample"
+		"LiveLinkProviderSample",
+		"SpatialMetricsProfiler"
 	],
 	"Modules" :
 	[
@@ -45,7 +46,8 @@
 				"CrashReportClientEditor",
 				"LiveLinkHub",
 				"TimeSynchronizationWriter",
-                "LiveLinkProviderSample"
+				"LiveLinkProviderSample",
+				"SpatialMetricsProfiler"
 			]
 		}
 	]
diff --git a/Engine/Plugins/WorldMetrics/WorldMetrics.uplugin b/Engine/Plugins/WorldMetrics/WorldMetrics.uplugin
--- a/Engine/Plugins/WorldMetrics/WorldMetrics.uplugin
+++ b/Engine/Plugins/WorldMetrics/WorldMetrics.uplugin
@@ -16,21 +16,24 @@
 	"EnabledByDefault": true,
 	"Installed": false,
 	"EditorCustomVirtualPath": "WorldMetrics",
+	"SupportedPrograms": [
+		"SpatialMetricsProfiler"
+	],
 	"Modules": [
 		{
 			"Name": "WorldMetricsCore",
-			"Type": "Runtime",
+			"Type": "RuntimeAndProgram",
 			"LoadingPhase": "Default"
 		},
 		{
 			"Name": "WorldMetricsTest",
-			"Type": "Runtime",
+			"Type": "RuntimeAndProgram",
 			"LoadingPhase": "Default",
 			"TargetConfigurationDenyList": [ "Shipping" ]
 		},
 		{
 			"Name": "CsvMetrics",
-			"Type": "Runtime",
+			"Type": "RuntimeAndProgram",
 			"LoadingPhase": "Default"
 		}
 	]
