CsvActorCountMetric - add missing include.

diff --git a/Engine/Plugins/WorldMetrics/Source/CsvMetrics/Private/CsvActorCountMetric.cpp b/Engine/Plugins/WorldMetrics/Source/CsvMetrics/Private/CsvActorCountMetric.cpp
--- a/Engine/Plugins/WorldMetrics/Source/CsvMetrics/Private/CsvActorCountMetric.cpp
+++ b/Engine/Plugins/WorldMetrics/Source/CsvMetrics/Private/CsvActorCountMetric.cpp
@@ -6,6 +6,7 @@
 #include "HAL/IConsoleManager.h"
 #include "ProfilingDebugging/CsvProfiler.h"
 #include "WorldMetricsActorTracker.h"
+#include "WorldMetricsSubsystem.h"
 
 #include UE_INLINE_GENERATED_CPP_BY_NAME(CsvActorCountMetric)
 
