[WorldMetrics] Fix LowLevelTests compilation

- Update test assertions to explicit nullptr comparisons
- Add GEngine nullptr validation before use in ScopedWorldTest
- Fix lambda captures by removing unnecessary [this] capture in static test contexts

#Synced-CL 50324810

diff --git a/Engine/Plugins/WorldMetrics/Source/WorldMetricsTest/Private/WorldMetricsTest.cpp b/Engine/Plugins/WorldMetrics/Source/WorldMetricsTest/Private/WorldMetricsTest.cpp
--- a/Engine/Plugins/WorldMetrics/Source/WorldMetricsTest/Private/WorldMetricsTest.cpp
+++ b/Engine/Plugins/WorldMetrics/Source/WorldMetricsTest/Private/WorldMetricsTest.cpp
@@ -196,7 +196,7 @@ int32 ReleaseMockMetricExtensions(
 void TestZeroState(UWorld* World)
 {
 	UWorldMetricsSubsystem* Subsystem = UWorldMetricsSubsystem::Get(World);
-	REQUIRE_MESSAGE(TEXT("WorldMetricsSubsystem is valid after GetSubsystem."), Subsystem);
+	REQUIRE_MESSAGE(TEXT("WorldMetricsSubsystem is valid after GetSubsystem."), Subsystem != nullptr);
 
 	REQUIRE_MESSAGE(
 		TEXT("WorldMetricsSubsystem shouldn't be enabled if there are no metrics."), !Subsystem->IsEnabled());
@@ -292,7 +292,7 @@ void TestSingleMetricSingleExtensionAcquireRelease(UWorld* World)
 			TEXT("AcquireExtension should fail"), !Subsystem->AcquireExtension(MetricA, UObject::StaticClass()));
 
 		UMockWorldMetricsExtensionA* ExtensionA = Subsystem->AcquireExtension<UMockWorldMetricsExtensionA>(MetricA);
-		REQUIRE_MESSAGE(TEXT("Acquire Test Extension failed."), ExtensionA);
+		REQUIRE_MESSAGE(TEXT("Acquire Test Extension failed."), ExtensionA != nullptr);
 		REQUIRE_MESSAGE(TEXT("Test Extension initialization failed"), ExtensionA->InitializeCount == 1);
 
 		REQUIRE_MESSAGE(
@@ -314,7 +314,7 @@ void TestSingleMetricSingleExtensionAcquireRelease(UWorld* World)
 	{
 		UMockWorldMetricsExtensionA* ExtensionA =
 			MetricA->GetOwner().AcquireExtension<UMockWorldMetricsExtensionA>(MetricA);
-		REQUIRE_MESSAGE(TEXT("Acquire Test Extension failed."), ExtensionA);
+		REQUIRE_MESSAGE(TEXT("Acquire Test Extension failed."), ExtensionA != nullptr);
 		REQUIRE_MESSAGE(TEXT("Test Extension initialization failed"), ExtensionA->InitializeCount == 1);
 
 		REQUIRE_MESSAGE(
@@ -378,7 +378,7 @@ void TestMultipleMetricsSingleExtensionAcquireReleaseChildExtension(UWorld* Worl
 	// Acquire Extension A for Metric A.
 	{
 		ExtensionA = MetricA->GetOwner().AcquireExtension<UMockWorldMetricsExtensionA>(MetricA);
-		REQUIRE_MESSAGE(TEXT("Acquire Test Extension A failed."), ExtensionA);
+		REQUIRE_MESSAGE(TEXT("Acquire Test Extension A failed."), ExtensionA != nullptr);
 		REQUIRE_MESSAGE(TEXT("Test Extension A initialization failed"), ExtensionA->InitializeCount == 1);
 
 		REQUIRE_MESSAGE(TEXT("WorldMetricsSubsystem should have extensions."), Subsystem->HasAnyExtension());
@@ -388,10 +388,10 @@ void TestMultipleMetricsSingleExtensionAcquireReleaseChildExtension(UWorld* Worl
 	// Acquire Extension B (and C, transitively) for Metric B.
 	{
 		ExtensionB = MetricB->GetOwner().AcquireExtension<UMockWorldMetricsExtensionB>(MetricB);
-		REQUIRE_MESSAGE(TEXT("Acquire Test Extension B failed."), ExtensionB);
+		REQUIRE_MESSAGE(TEXT("Acquire Test Extension B failed."), ExtensionB != nullptr);
 		REQUIRE_MESSAGE(TEXT("Test Extension B initialization failed"), ExtensionB->InitializeCount == 1);
 
-		REQUIRE_MESSAGE(TEXT("Acquire Test Extension C failed."), ExtensionC);
+		REQUIRE_MESSAGE(TEXT("Acquire Test Extension C failed."), ExtensionC != nullptr);
 		REQUIRE_MESSAGE(TEXT("Test Extension C initialization failed"), ExtensionC->InitializeCount == 1);
 
 		REQUIRE_MESSAGE(TEXT("WorldMetricsSubsystem should have extensions."), Subsystem->HasAnyExtension());
@@ -443,7 +443,7 @@ void TestMultipleMetricsSingleExtensionAcquireRelease(UWorld* World, bool bRando
 		for (UMockWorldMetricBase* MockMetric : MockMetrics)
 		{
 			Subsystem->AcquireExtension<UMockWorldMetricsExtensionA>(MockMetric);
-			REQUIRE_MESSAGE(TEXT("Acquire Test Extension failed."), ExtensionA);
+			REQUIRE_MESSAGE(TEXT("Acquire Test Extension failed."), ExtensionA != nullptr);
 			REQUIRE_MESSAGE(TEXT("Test Extension initialization failed"), ExtensionA->InitializeCount == 1);
 			REQUIRE_MESSAGE(TEXT("Test Extension initialization failed"), ExtensionA->OnAcquireCount == ++CheckCount);
 		}
@@ -461,7 +461,7 @@ void TestMultipleMetricsSingleExtensionAcquireRelease(UWorld* World, bool bRando
 			REQUIRE_MESSAGE(TEXT("Test Extension unexpected deinitialization"), ExtensionA->DeinitializeCount == 0);
 			Subsystem->ReleaseExtension<UMockWorldMetricsExtensionA>(MockMetric);
 
-			REQUIRE_MESSAGE(TEXT("Acquire Test Extension failed."), ExtensionA);
+			REQUIRE_MESSAGE(TEXT("Acquire Test Extension failed."), ExtensionA != nullptr);
 			REQUIRE_MESSAGE(TEXT("Test Extension initialization failed"), ExtensionA->InitializeCount == 1);
 			REQUIRE_MESSAGE(TEXT("Test Extension initialization failed"), ExtensionA->OnReleaseCount == ++CheckCount);
 		}
@@ -521,7 +521,7 @@ void TestSingleMetricSingleExtensionAutoReleaseOnRemoval(UWorld* World)
 	REQUIRE_MESSAGE(TEXT("ContainsMetric failed"), Subsystem->ContainsMetric(MetricA));
 
 	UMockWorldMetricsExtensionA* ExtensionA = Subsystem->AcquireExtension<UMockWorldMetricsExtensionA>(MetricA);
-	REQUIRE_MESSAGE(TEXT("Acquire Test Extension failed."), ExtensionA);
+	REQUIRE_MESSAGE(TEXT("Acquire Test Extension failed."), ExtensionA != nullptr);
 	REQUIRE_MESSAGE(TEXT("Test Extension initialization failed"), ExtensionA->InitializeCount == 1);
 
 	REQUIRE_MESSAGE(
@@ -573,12 +573,12 @@ void TestSingleMetricOrphanExtensionAutoRemoval(UWorld* World)
 	REQUIRE_MESSAGE(TEXT("ContainsMetric failed"), Subsystem->ContainsMetric(MetricA));
 
 	UMockWorldMetricsExtensionA* ExtensionA = Subsystem->AcquireExtension<UMockWorldMetricsExtensionA>(MetricA);
-	REQUIRE_MESSAGE(TEXT("Acquire Test Extension failed."), ExtensionA);
+	REQUIRE_MESSAGE(TEXT("Acquire Test Extension failed."), ExtensionA != nullptr);
 	REQUIRE_MESSAGE(TEXT("Test Extension initialization failed"), ExtensionA->InitializeCount == 1);
 
 	// Extension dependency: ExtensionB requires ExtensionA
 	UMockWorldMetricsExtensionB* ExtensionB = Subsystem->AcquireExtension<UMockWorldMetricsExtensionB>(ExtensionA);
-	REQUIRE_MESSAGE(TEXT("Acquire Test Extension failed."), ExtensionB);
+	REQUIRE_MESSAGE(TEXT("Acquire Test Extension failed."), ExtensionB != nullptr);
 	REQUIRE_MESSAGE(TEXT("Test Extension initialization failed"), ExtensionB->InitializeCount == 1);
 
 	REQUIRE_MESSAGE(
@@ -803,6 +803,8 @@ namespace Private
 {
 static void ScopedWorldTest(EWorldType::Type WorldType, TFunctionRef<void(UWorld* World)> WorldTest)
 {
+	REQUIRE_MESSAGE(TEXT("GEngine is required for this test"), GEngine != nullptr);
+	
 	constexpr bool bInformEngineOfWorld = false;
 	UWorld* World = UWorld::CreateWorld(WorldType, bInformEngineOfWorld);
 	REQUIRE_MESSAGE(TEXT("World could not be created."), World != nullptr);
@@ -824,7 +826,7 @@ static void ScopedWorldTest(EWorldType::Type WorldType, TFunctionRef<void(UWorld
 
 TEST_CASE_NAMED(WorldMetricsTestZeroState, "WorldMetrics::TestZeroState", "[WorldMetrics][Core][EngineFilter]")
 {
-	Private::ScopedWorldTest(EWorldType::Editor, [this](UWorld* World) { Private::TestZeroState(World); });
+	Private::ScopedWorldTest(EWorldType::Editor, [](UWorld* World) { Private::TestZeroState(World); });
 }
 
 TEST_CASE_NAMED(
@@ -832,7 +834,7 @@ TEST_CASE_NAMED(
 	"WorldMetrics::TestSingleMetricAddRemove",
 	"[WorldMetrics][Core][EngineFilter]")
 {
-	Private::ScopedWorldTest(EWorldType::Editor, [this](UWorld* World) { Private::TestSingleMetricAddRemove(World); });
+	Private::ScopedWorldTest(EWorldType::Editor, [](UWorld* World) { Private::TestSingleMetricAddRemove(World); });
 }
 
 TEST_CASE_NAMED(
@@ -842,7 +844,7 @@ TEST_CASE_NAMED(
 {
 	Private::ScopedWorldTest(
 		EWorldType::Editor,
-		[this](UWorld* World)
+		[](UWorld* World)
 		{
 			Private::TestMultipleMetricsAddRemove(World, false);
 			Private::TestMultipleMetricsAddRemove(World, true);
@@ -856,7 +858,7 @@ TEST_CASE_NAMED(
 {
 	Private::ScopedWorldTest(
 		EWorldType::Editor,
-		[this](UWorld* World)
+		[](UWorld* World)
 		{
 			Private::TestMultipleMetricIteration(World, false);
 			Private::TestMultipleMetricIteration(World, true);
@@ -869,7 +871,7 @@ TEST_CASE_NAMED(
 	"[WorldMetrics][Core][EngineFilter]")
 {
 	Private::ScopedWorldTest(
-		EWorldType::Editor, [this](UWorld* World) { Private::TestSingleMetricSingleExtensionAcquireRelease(World); });
+		EWorldType::Editor, [](UWorld* World) { Private::TestSingleMetricSingleExtensionAcquireRelease(World); });
 }
 
 TEST_CASE_NAMED(
@@ -879,7 +881,7 @@ TEST_CASE_NAMED(
 {
 	Private::ScopedWorldTest(
 		EWorldType::Editor,
-		[this](UWorld* World)
+		[](UWorld* World)
 		{
 			Private::TestMultipleMetricsSingleExtensionAcquireRelease(World, false);
 			Private::TestMultipleMetricsSingleExtensionAcquireRelease(World, true);
@@ -893,7 +895,7 @@ TEST_CASE_NAMED(
 {
 	Private::ScopedWorldTest(
 		EWorldType::Editor,
-		[this](UWorld* World)
+		[](UWorld* World)
 		{
 			Private::TestMultipleMetricsSingleExtensionAcquireReleaseChildExtension(World);
 		});
@@ -906,7 +908,7 @@ TEST_CASE_NAMED(
 {
 	Private::ScopedWorldTest(
 		EWorldType::Editor,
-		[this](UWorld* World)
+		[](UWorld* World)
 		{
 			Private::TestMultipleMetricsMultipleExtensionAcquireRelease(World, false);
 			Private::TestMultipleMetricsMultipleExtensionAcquireRelease(World, true);
@@ -920,7 +922,7 @@ TEST_CASE_NAMED(
 {
 	Private::ScopedWorldTest(
 		EWorldType::Editor,
-		[this](UWorld* World) { Private::TestSingleMetricSingleExtensionAutoReleaseOnRemoval(World); });
+		[](UWorld* World) { Private::TestSingleMetricSingleExtensionAutoReleaseOnRemoval(World); });
 }
 
 TEST_CASE_NAMED(
@@ -930,7 +932,7 @@ TEST_CASE_NAMED(
 {
 	Private::ScopedWorldTest(
 		EWorldType::Editor,
-		[this](UWorld* World)
+		[](UWorld* World)
 		{
 			Private::TestMultipleMetricsMultipleExtensionAutoReleaseOnRemoval(World, false);
 			Private::TestMultipleMetricsMultipleExtensionAutoReleaseOnRemoval(World, true);
@@ -943,7 +945,7 @@ TEST_CASE_NAMED(
 	"[WorldMetrics][Core][EngineFilter]")
 {
 	Private::ScopedWorldTest(
-		EWorldType::Editor, [this](UWorld* World) { Private::TestSingleMetricOrphanExtensionAutoRemoval(World); });
+		EWorldType::Editor, [](UWorld* World) { Private::TestSingleMetricOrphanExtensionAutoRemoval(World); });
 }
 
 TEST_CASE_NAMED(
@@ -953,7 +955,7 @@ TEST_CASE_NAMED(
 {
 	Private::ScopedWorldTest(
 		EWorldType::Editor,
-		[this](UWorld* World)
+		[](UWorld* World)
 		{
 			Private::TestMetricCollection(World, true);
 			Private::TestMetricCollection(World, false);
